[MassDebugger] added missing details about ConstSharedFragments in the archetype panel

diff --git a/Engine/Source/Editor/MassEntityDebugger/Private/SMassArchetype.cpp b/Engine/Source/Editor/MassEntityDebugger/Private/SMassArchetype.cpp
--- a/Engine/Source/Editor/MassEntityDebugger/Private/SMassArchetype.cpp
+++ b/Engine/Source/Editor/MassEntityDebugger/Private/SMassArchetype.cpp
@@ -105,6 +105,7 @@ void SMassArchetype::Construct(const FArguments& InArgs, TSharedPtr<FMassDebugge
 		UE::Mass::Debugger::UI::AddBitSetDiff(Box, ParentComposition.Tags, Composition.Tags, TEXT("Tags"), Brush, Prune);
 		UE::Mass::Debugger::UI::AddBitSetDiff(Box, ParentComposition.ChunkFragments, Composition.ChunkFragments, TEXT("Chunk Fragments"), Brush, Prune);
 		UE::Mass::Debugger::UI::AddBitSetDiff(Box, ParentComposition.SharedFragments, Composition.SharedFragments, TEXT("Shared Fragments"), Brush, Prune);
+		UE::Mass::Debugger::UI::AddBitSetDiff(Box, ParentComposition.ConstSharedFragments, Composition.ConstSharedFragments, TEXT("Const Shared Fragments"), Brush, Prune);
 	}
 	else
 	{
@@ -112,6 +113,7 @@ void SMassArchetype::Construct(const FArguments& InArgs, TSharedPtr<FMassDebugge
 		UE::Mass::Debugger::UI::AddBitSet(Box, Composition.Tags, TEXT("Tags"), Brush);
 		UE::Mass::Debugger::UI::AddBitSet(Box, Composition.ChunkFragments, TEXT("Chunk Fragments"), Brush);
 		UE::Mass::Debugger::UI::AddBitSet(Box, Composition.SharedFragments, TEXT("Shared Fragments"), Brush);
+		UE::Mass::Debugger::UI::AddBitSet(Box, Composition.ConstSharedFragments, TEXT("Const Shared Fragments"), Brush);
 	}
 
 	ChildSlot
