[Mass] Issue 624880: added a bunch of World and other headers in hopes of fixing the non-unit compilation issue.

#rnx

diff --git a/Engine/Plugins/AI/MassAI/Source/MassNavigationEditor/Private/MassNavigationTestingActor.cpp b/Engine/Plugins/AI/MassAI/Source/MassNavigationEditor/Private/MassNavigationTestingActor.cpp
--- a/Engine/Plugins/AI/MassAI/Source/MassNavigationEditor/Private/MassNavigationTestingActor.cpp
+++ b/Engine/Plugins/AI/MassAI/Source/MassNavigationEditor/Private/MassNavigationTestingActor.cpp
@@ -7,6 +7,7 @@
 #include "ZoneGraphDelegates.h"
 #include "ZoneGraphRenderingUtilities.h"
 #include "ZoneGraphData.h"
+#include "Engine/World.h"
 
 #if UE_ENABLE_DEBUG_DRAWING
 
diff --git a/Engine/Plugins/Runtime/InstancedActors/Source/InstancedActors/Private/GameplayDebuggerCategory_InstancedActors.cpp b/Engine/Plugins/Runtime/InstancedActors/Source/InstancedActors/Private/GameplayDebuggerCategory_InstancedActors.cpp
--- a/Engine/Plugins/Runtime/InstancedActors/Source/InstancedActors/Private/GameplayDebuggerCategory_InstancedActors.cpp
+++ b/Engine/Plugins/Runtime/InstancedActors/Source/InstancedActors/Private/GameplayDebuggerCategory_InstancedActors.cpp
@@ -19,7 +19,6 @@
 #include "MassEntitySubsystem.h"
 #include "MassEntityView.h"
 #include "Engine/World.h"
-#include "MassEntitySubsystem.h"
 
 
 namespace UE::InstancedActors::Debug
diff --git a/Engine/Source/Editor/MassEntityDebugger/Private/MassDebuggerModel.cpp b/Engine/Source/Editor/MassEntityDebugger/Private/MassDebuggerModel.cpp
--- a/Engine/Source/Editor/MassEntityDebugger/Private/MassDebuggerModel.cpp
+++ b/Engine/Source/Editor/MassEntityDebugger/Private/MassDebuggerModel.cpp
@@ -10,6 +10,7 @@
 #include "UObject/UObjectIterator.h"
 #include "Containers/UnrealString.h"
 #include "MassArchetypeData.h"
+#include "Engine/World.h"
 
 #define LOCTEXT_NAMESPACE "SMassDebugger"
 
@@ -367,6 +368,15 @@ FMassDebuggerProcessingGraph::FMassDebuggerProcessingGraph(const FMassDebuggerMo
 #endif // WITH_MASSENTITY_DEBUG
 }
 
+//----------------------------------------------------------------------//
+// FMassDebuggerModel
+//----------------------------------------------------------------------//
+FMassDebuggerEnvironment::FMassDebuggerEnvironment(const TSharedRef<const FMassEntityManager>& InEntityManager)
+	: EntityManager(InEntityManager), World(InEntityManager->GetWorld()), bNeedsValidWorld(InEntityManager->GetWorld() != nullptr)
+{
+	
+}
+
 //----------------------------------------------------------------------//
 // FMassDebuggerModel
 //----------------------------------------------------------------------//
diff --git a/Engine/Source/Editor/MassEntityDebugger/Private/MassDebuggerModel.h b/Engine/Source/Editor/MassEntityDebugger/Private/MassDebuggerModel.h
--- a/Engine/Source/Editor/MassEntityDebugger/Private/MassDebuggerModel.h
+++ b/Engine/Source/Editor/MassEntityDebugger/Private/MassDebuggerModel.h
@@ -14,6 +14,7 @@ class UMassProcessor;
 struct FMassArchetypeHandle;
 struct FMassDebuggerModel;
 struct FMassEntityManager;
+class UWorld;
 
 enum class EMassDebuggerSelectionMode : uint8
 {
@@ -139,9 +140,7 @@ struct FMassDebuggerProcessingGraph
 
 struct FMassDebuggerEnvironment
 {
-	explicit FMassDebuggerEnvironment(const TSharedRef<const FMassEntityManager>& InEntityManager)
-		: EntityManager(InEntityManager), World(InEntityManager->GetWorld()), bNeedsValidWorld(InEntityManager->GetWorld() != nullptr)
-	{}
+	explicit FMassDebuggerEnvironment(const TSharedRef<const FMassEntityManager>& InEntityManager);
 
 	bool operator==(const FMassDebuggerEnvironment& Other) const { return EntityManager == Other.EntityManager; }
 
