[Mass] Changes to how MassDebugger collects its information about processors. Now everybody can register as a collection of processors! The collected processors are also nicely grouped in the Processors View.

diff --git a/Engine/Source/Editor/MassEntityDebugger/Private/MassDebuggerModel.cpp b/Engine/Source/Editor/MassEntityDebugger/Private/MassDebuggerModel.cpp
--- a/Engine/Source/Editor/MassEntityDebugger/Private/MassDebuggerModel.cpp
+++ b/Engine/Source/Editor/MassEntityDebugger/Private/MassDebuggerModel.cpp
@@ -4,6 +4,7 @@
 #include "MassProcessor.h"
 #include "MassEntityManager.h"
 #include "MassEntityQuery.h"
+#include "MassProcessingPhaseManager.h"
 #include "Engine/Engine.h"
 #include "MassDebugger.h"
 #include "MassDebuggerSettings.h"
@@ -42,7 +43,8 @@ namespace UE::Mass::Debugger::Private
 		{
 			return;
 		}
-		
+
+		OutDisplayName.RemoveFromStart(TEXT("Default__"), ESearchCase::CaseSensitive);
 		OutDisplayName.RemoveFromStart(TEXT("Mass"), ESearchCase::CaseSensitive);
 	}
 
@@ -72,7 +74,7 @@ FString FMassDebuggerEnvironment::GetDisplayName() const
 	FString DisplayName;
 
 #if WITH_MASSENTITY_DEBUG
-	if (const FMassEntityManager* EntityManagerPtr = GetEntityManager())
+	if (TSharedPtr<const FMassEntityManager> EntityManagerPtr = GetEntityManager())
 	{
 		DisplayName += EntityManagerPtr->DebugGetName();
 		if (DisplayName.Len())
@@ -87,9 +89,9 @@ FString FMassDebuggerEnvironment::GetDisplayName() const
 	return DisplayName;
 }
 
-const FMassEntityManager* FMassDebuggerEnvironment::GetEntityManager() const
+TSharedPtr<const FMassEntityManager> FMassDebuggerEnvironment::GetEntityManager() const
 {
-	return EntityManager.Pin().Get();
+	return EntityManager.Pin();
 }
 
 //----------------------------------------------------------------------//
@@ -141,12 +143,14 @@ FMassDebuggerProcessorData::FMassDebuggerProcessorData(const UMassProcessor& InP
 #endif // WITH_MASSENTITY_DEBUG
 }
 
-FMassDebuggerProcessorData::FMassDebuggerProcessorData(const FMassEntityManager& EntityManager, UMassProcessor& InProcessor
+FMassDebuggerProcessorData::FMassDebuggerProcessorData(const FMassEntityManager& EntityManager, const UMassProcessor& InProcessor
 	, const TMap<FMassArchetypeHandle, TSharedPtr<FMassDebuggerArchetypeData>>& InTransientArchetypesMap)
 {
 	SetProcessor(InProcessor);
 #if WITH_MASSENTITY_DEBUG
-	TConstArrayView<FMassEntityQuery*> ProcessorQueries = FMassDebugger::GetUpToDateProcessorQueries(EntityManager, InProcessor);
+	// yeah, it's ugly. But it's debugging code, so... 
+	UMassProcessor& MutableProcessor = const_cast<UMassProcessor&>(InProcessor);
+	TConstArrayView<FMassEntityQuery*> ProcessorQueries = FMassDebugger::GetUpToDateProcessorQueries(EntityManager, MutableProcessor);
 
 	ProcessorRequirements = MakeShareable(new FMassDebuggerQueryData(InProcessor.GetProcessorRequirements(), LOCTEXT("MassProcessorRequirementsLabel", "Processor Requirements")));
 
@@ -369,7 +373,7 @@ FMassDebuggerProcessingGraph::FMassDebuggerProcessingGraph(const FMassDebuggerMo
 }
 
 //----------------------------------------------------------------------//
-// FMassDebuggerModel
+// FMassDebuggerEnvironment
 //----------------------------------------------------------------------//
 FMassDebuggerEnvironment::FMassDebuggerEnvironment(const TSharedRef<const FMassEntityManager>& InEntityManager)
 	: EntityManager(InEntityManager), World(InEntityManager->GetWorld()), bNeedsValidWorld(InEntityManager->GetWorld() != nullptr)
@@ -387,6 +391,8 @@ FMassDebuggerModel::FMassDebuggerModel()
 #endif // WITH_MASSENTITY_DEBUG
 }
 
+// disabling to avoid reporting CachedProcessors deprecation
+PRAGMA_DISABLE_DEPRECATION_WARNINGS
 FMassDebuggerModel::~FMassDebuggerModel()
 {
 #if WITH_MASSENTITY_DEBUG
@@ -396,6 +402,7 @@ FMassDebuggerModel::~FMassDebuggerModel()
 	}
 #endif // WITH_MASSENTITY_DEBUG
 }
+PRAGMA_ENABLE_DEPRECATION_WARNINGS
 
 void FMassDebuggerModel::SetEnvironment(const TSharedPtr<FMassDebuggerEnvironment>& Item)
 {
@@ -415,15 +422,18 @@ void FMassDebuggerModel::SetEnvironment(const TSharedPtr<FMassDebuggerEnvironmen
 
 void FMassDebuggerModel::RefreshAll()
 {
-	TMap<FMassArchetypeHandle, TSharedPtr<FMassDebuggerArchetypeData>> TransientArchetypesMap;
+	if (Environment)
+	{
+		TMap<FMassArchetypeHandle, TSharedPtr<FMassDebuggerArchetypeData>> TransientArchetypesMap;
 
-	CacheArchetypesData(TransientArchetypesMap);
-	CacheProcessorsData(TransientArchetypesMap);
-	CacheProcessingGraphs();
+		CacheArchetypesData(TransientArchetypesMap);
+		CacheProcessorsData(TransientArchetypesMap);
+		CacheProcessingGraphs();
 
-	ClearArchetypeSelection();
+		ClearArchetypeSelection();
 
-	OnRefreshDelegate.Broadcast();
+		OnRefreshDelegate.Broadcast();
+	}
 }
 
 void FMassDebuggerModel::SelectProcessor(TSharedPtr<FMassDebuggerProcessorData>& Processor)
@@ -473,7 +483,7 @@ void FMassDebuggerModel::SelectArchetypes(TArrayView<TSharedPtr<FMassDebuggerArc
 
 	SelectedArchetypes = InSelectedArchetypes;
 
-	for (TSharedPtr<FMassDebuggerProcessorData>& ProcessorData : CachedProcessors)
+	for (TSharedPtr<FMassDebuggerProcessorData>& ProcessorData : AllCachedProcessors)
 	{
 		check(ProcessorData.IsValid());
 		for (const TSharedPtr<FMassDebuggerArchetypeData>& ArchetypeData : InSelectedArchetypes)
@@ -500,26 +510,57 @@ void FMassDebuggerModel::ClearArchetypeSelection()
 
 void FMassDebuggerModel::CacheProcessorsData(const TMap<FMassArchetypeHandle, TSharedPtr<FMassDebuggerArchetypeData>>& InTransientArchetypesMap)
 {
-	CachedProcessors.Reset();
+	static auto SortPredicate = [](const TSharedPtr<FMassDebuggerProcessorData>& A, const TSharedPtr<FMassDebuggerProcessorData>& B)
+	{
+		return A->Label < B->Label;
+	};
 
-	UWorld* World = Environment ? Environment->World.Get() : nullptr;
-	const FMassEntityManager* EntityManager = Environment ? Environment->GetEntityManager() : nullptr;
+	CachedProcessorCollections.Reset();
+	AllCachedProcessors.Reset();
+	
+	if (Environment == nullptr)
+	{
+		return;
+	}
 
-	if (EntityManager)
+	UWorld* World = Environment->World.Get();
+
+	if (TSharedPtr<const FMassEntityManager> EntityManager = Environment->GetEntityManager())
 	{
-		for (FThreadSafeObjectIterator It(UMassProcessor::StaticClass()); It; ++It)
+		// run all the processor providers and convert the data to FMassDebuggerProcessorData
+		TArray<const UMassProcessor*> TmpProcessors;
+		for (auto& Providers : Environment->ProcessorProviders)
 		{
-			UMassProcessor* Processor = Cast<UMassProcessor>(*It);
-			if (Processor 
-				&& Cast<UMassCompositeProcessor>(Processor) == nullptr
-				&& UE::Mass::Debugger::Private::IsDebuggableProcessor(World, *Processor))
+			
+			// this will fill TmpProcessors with results of the stored FMassDebugger::FProcessorProviderFunction,
+			// the "Value" is of type FMassDebugger::FProcessorProviderFunction
+			Providers.Value(TmpProcessors);
+			if (TmpProcessors.IsEmpty())
 			{
-				CachedProcessors.Add(MakeShareable(new FMassDebuggerProcessorData(*EntityManager, *Processor, InTransientArchetypesMap)));
+				continue;
 			}
+
+			TSharedPtr<FProcessorCollection>& Collection = CachedProcessorCollections.Add_GetRef(MakeShareable(new FProcessorCollection(Providers.Key)));
+			TArray<TSharedPtr<FMassDebuggerProcessorData>>& Container = Collection->Container;
+			Container.Reserve(TmpProcessors.Num());
+			for (const UMassProcessor* Processor : TmpProcessors)
+			{
+				if (Processor)
+				{
+					AllCachedProcessors.Add(
+						Container.Add_GetRef(MakeShareable(new FMassDebuggerProcessorData(*EntityManager, *Processor, InTransientArchetypesMap)))
+					);
+				}
+			}
+			Container.Sort(SortPredicate);
+
+			TmpProcessors.Reset();
 		}
 	}
 	else
 	{
+		TSharedPtr<FProcessorCollection>& Collection = CachedProcessorCollections.Add_GetRef(MakeShareable(new FProcessorCollection(TEXT("Global view"))));
+		TArray<TSharedPtr<FMassDebuggerProcessorData>>& Container = Collection->Container;
 		for (FThreadSafeObjectIterator It(UMassProcessor::StaticClass()); It; ++It)
 		{
 			UMassProcessor* Processor = Cast<UMassProcessor>(*It);
@@ -527,15 +568,13 @@ void FMassDebuggerModel::CacheProcessorsData(const TMap<FMassArchetypeHandle, TS
 				&& Cast<UMassCompositeProcessor>(Processor) == nullptr
 				&& UE::Mass::Debugger::Private::IsDebuggableProcessor(World, *Processor))
 			{
-				CachedProcessors.Add(MakeShareable(new FMassDebuggerProcessorData(*Processor)));
+				Container.Add(MakeShareable(new FMassDebuggerProcessorData(*Processor)));
 			}
 		}
+		Container.Sort(SortPredicate);
 	}
 
-	CachedProcessors.Sort([](const TSharedPtr<FMassDebuggerProcessorData>& A, const TSharedPtr<FMassDebuggerProcessorData>& B)
-	{
-		return A->Label < B->Label;
-	});
+	AllCachedProcessors.Sort(SortPredicate);
 }
 
 void FMassDebuggerModel::CacheProcessingGraphs()
@@ -543,10 +582,10 @@ void FMassDebuggerModel::CacheProcessingGraphs()
 	CachedProcessingGraphs.Reset();
 
 	UWorld* World = Environment ? Environment->World.Get() : nullptr;
-	for (FThreadSafeObjectIterator It(UMassProcessor::StaticClass()); It; ++It)
+	for (FThreadSafeObjectIterator It(UMassCompositeProcessor::StaticClass()); It; ++It)
 	{
-		UMassCompositeProcessor* Processor = Cast<UMassCompositeProcessor>(*It);
-		if (Processor && UE::Mass::Debugger::Private::IsDebuggableProcessor(World, *Processor))
+		UMassCompositeProcessor* Processor = CastChecked<UMassCompositeProcessor>(*It);
+		if (UE::Mass::Debugger::Private::IsDebuggableProcessor(World, *Processor))
 		{
 			CachedProcessingGraphs.Add(MakeShareable(new FMassDebuggerProcessingGraph(*this, *Processor)));
 		}
@@ -560,7 +599,7 @@ void FMassDebuggerModel::CacheArchetypesData(TMap<FMassArchetypeHandle, TSharedP
 
 	if (Environment)
 	{
-		if (const FMassEntityManager* EntityManager = Environment->GetEntityManager())
+		if (TSharedPtr<const FMassEntityManager> EntityManager = Environment->GetEntityManager())
 		{
 			StoreArchetypes(*EntityManager, OutTransientArchetypesMap);
 		}
@@ -667,14 +706,16 @@ bool FMassDebuggerModel::IsStale() const
 
 const TSharedPtr<FMassDebuggerProcessorData>& FMassDebuggerModel::GetProcessorDataChecked(const UMassProcessor& Processor) const
 {
-	check(CachedProcessors.Num());
+	check(AllCachedProcessors.Num());
 
 	const uint32 ProcessorHash = UE::Mass::Debugger::Private::CalcProcessorHash(Processor);
-
-	const TSharedPtr<FMassDebuggerProcessorData>* DataFound = CachedProcessors.FindByPredicate([ProcessorHash](const TSharedPtr<FMassDebuggerProcessorData>& Element)
+	auto SearchPredicate = [ProcessorHash](const TSharedPtr<FMassDebuggerProcessorData>& Element)
 		{
 			return Element->ProcessorHash == ProcessorHash;
-		});
+		};
+
+	// @todo could convert AllCachedProcessors to a map if this search becomes too slow
+	const TSharedPtr<FMassDebuggerProcessorData>* DataFound = AllCachedProcessors.FindByPredicate(SearchPredicate);
 
 	check(DataFound);
 	return *DataFound;
@@ -691,8 +732,8 @@ void FMassDebuggerModel::ResetSelectedArchetypes()
 
 void FMassDebuggerModel::ResetSelectedProcessors()
 {
-	// using CachedProcessors instead of SelectedProcessors to be on the safe side
-	for (TSharedPtr<FMassDebuggerProcessorData>& ProcessorData : CachedProcessors)
+	// resetting marking of all the processors instead of SelectedProcessors to be on the safe side
+	for (TSharedPtr<FMassDebuggerProcessorData>& ProcessorData : AllCachedProcessors)
 	{
 		check(ProcessorData.IsValid());
 		ProcessorData->Selection = EMassDebuggerProcessorSelection::None;
@@ -702,7 +743,7 @@ void FMassDebuggerModel::ResetSelectedProcessors()
 
 void FMassDebuggerModel::OnEntitySelected(const FMassEntityManager& EntityManager, const FMassEntityHandle EntityHandle)
 {
-	if (!Environment || Environment->GetEntityManager() != &EntityManager)
+	if (!Environment || Environment->GetEntityManager().Get() != &EntityManager)
 	{
 		// not the entity manager we're debugging right now
 		return;
diff --git a/Engine/Source/Editor/MassEntityDebugger/Private/MassDebuggerModel.h b/Engine/Source/Editor/MassEntityDebugger/Private/MassDebuggerModel.h
--- a/Engine/Source/Editor/MassEntityDebugger/Private/MassDebuggerModel.h
+++ b/Engine/Source/Editor/MassEntityDebugger/Private/MassDebuggerModel.h
@@ -97,7 +97,7 @@ struct FMassDebuggerArchetypeData
 struct FMassDebuggerProcessorData
 {
 	FMassDebuggerProcessorData(const UMassProcessor& InProcessor);
-	FMassDebuggerProcessorData(const FMassEntityManager& EntitySubsystem, UMassProcessor& InProcessor, const TMap<FMassArchetypeHandle, TSharedPtr<FMassDebuggerArchetypeData>>& InTransientArchetypesMap);
+	FMassDebuggerProcessorData(const FMassEntityManager& EntitySubsystem, const UMassProcessor& InProcessor, const TMap<FMassArchetypeHandle, TSharedPtr<FMassDebuggerArchetypeData>>& InTransientArchetypesMap);
 
 private:
 	void SetProcessor(const UMassProcessor& InProcessor);
@@ -145,11 +145,12 @@ struct FMassDebuggerEnvironment
 	bool operator==(const FMassDebuggerEnvironment& Other) const { return EntityManager == Other.EntityManager; }
 
 	FString GetDisplayName() const;
-	const FMassEntityManager* GetEntityManager() const;
+	TSharedPtr<const FMassEntityManager> GetEntityManager() const;
 	bool IsWorldValid() const { return World.IsValid(); }
 	bool NeedsValidWorld() const { return bNeedsValidWorld; }
 	
 	TWeakPtr<const FMassEntityManager> EntityManager;
+	TMap<FName, FMassDebugger::FProcessorProviderFunction> ProcessorProviders;
 	TWeakObjectPtr<UWorld> World;
 	bool bNeedsValidWorld = false;
 };
@@ -177,6 +178,7 @@ struct FMassDebuggerModel
 
 	bool IsCurrentEnvironment(const FMassDebuggerEnvironment& InEnvironment) const { return Environment && *Environment.Get() == InEnvironment; }
 	bool IsCurrentEnvironmentValid() const { return Environment && Environment->EntityManager.IsValid(); }
+	bool HasEnvironmentSelected() const { return static_cast<bool>(Environment); }
 
 	void CacheArchetypesData(TMap<FMassArchetypeHandle, TSharedPtr<FMassDebuggerArchetypeData>>& OutTransientArchetypesMap); 
 	void CacheProcessorsData(const TMap<FMassArchetypeHandle, TSharedPtr<FMassDebuggerArchetypeData>>& InTransientArchetypesMap);
@@ -207,7 +209,16 @@ public:
 	EMassDebuggerSelectionMode SelectionMode = EMassDebuggerSelectionMode::None;
 
 	TSharedPtr<FMassDebuggerEnvironment> Environment;
-	TArray<TSharedPtr<FMassDebuggerProcessorData>> CachedProcessors;
+	struct FProcessorCollection
+	{
+		FProcessorCollection(FName InLabel = FName())
+			: Label(InLabel)
+		{	
+		}
+		FName Label;
+		TArray<TSharedPtr<FMassDebuggerProcessorData>> Container;
+	};
+	TArray<TSharedPtr<FProcessorCollection>> CachedProcessorCollections;
 	TArray<TSharedPtr<FMassDebuggerProcessorData>> SelectedProcessors;
 
 	TArray<TSharedPtr<FMassDebuggerArchetypeData>> CachedAllArchetypes;
@@ -223,5 +234,12 @@ public:
 	FString EnvironmentDisplayName;
 
 	FDelegateHandle OnEntitySelectedHandle;
+
+protected:	
+	TArray<TSharedPtr<FMassDebuggerProcessorData>> AllCachedProcessors;
+
+public:
+	UE_DEPRECATED(5.6, "CachedProcessors property is now deprecated. Use CachedProcessorCollections instead. ")
+	TArray<TSharedPtr<FMassDebuggerProcessorData>> CachedProcessors;
 };
 
diff --git a/Engine/Source/Editor/MassEntityDebugger/Private/SMassDebugger.cpp b/Engine/Source/Editor/MassEntityDebugger/Private/SMassDebugger.cpp
--- a/Engine/Source/Editor/MassEntityDebugger/Private/SMassDebugger.cpp
+++ b/Engine/Source/Editor/MassEntityDebugger/Private/SMassDebugger.cpp
@@ -58,6 +58,7 @@ SMassDebugger::~SMassDebugger()
 #if WITH_MASSENTITY_DEBUG
 	FMassDebugger::OnEntityManagerInitialized.Remove(OnEntityManagerInitializedHandle);
 	FMassDebugger::OnEntityManagerDeinitialized.Remove(OnEntityManagerDeinitializedHandle);
+	FMassDebugger::OnProcessorProviderRegistered.Remove(OnProcessorProviderRegisteredHandle);
 #endif // WITH_MASSENTITY_DEBUG
 }
 
@@ -76,7 +77,7 @@ void SMassDebugger::Construct(const FArguments& InArgs, const TSharedRef<SDockTa
 				DebuggerModel->RefreshAll();
 		})
 		// check if there's any mass environment picked
-		, FCanExecuteAction::CreateLambda([this]()->bool { return true; })
+		, FCanExecuteAction::CreateLambda([this]()->bool { return DebuggerModel->HasEnvironmentSelected(); })
 		, FIsActionChecked::CreateLambda([]()->bool { return false; })
 		, FIsActionButtonVisible::CreateLambda([this]()->bool { return true; }));
 
@@ -202,6 +203,19 @@ void SMassDebugger::OnEntityManagerInitialized(const FMassEntityManager& EntityM
 	}
 }
 
+void SMassDebugger::OnProcessorProviderRegistered(const FMassDebugger::FEnvironment& Environment)
+{
+	const int32 FoundIndex = EnvironmentsList.FindLastByPredicate([WeakManager = Environment.EntityManager](const TSharedPtr<FMassDebuggerEnvironment>& TestedEnvironment)
+	{
+		return TestedEnvironment->EntityManager == WeakManager;
+	});
+	
+	if (ensureMsgf(FoundIndex != INDEX_NONE, TEXT("We never expect OnProcessorProviderRegistered called for an Environment that has not been registered")))
+	{
+		EnvironmentsList[FoundIndex]->ProcessorProviders = Environment.ProcessorProviders;
+	}
+}
+
 void SMassDebugger::OnEntityManagerDeinitialized(const FMassEntityManager& EntityManager)
 {
 	const UWorld* World = EntityManager.GetWorld();
@@ -253,6 +267,7 @@ void SMassDebugger::HandleEnvironmentChanged(TSharedPtr<FMassDebuggerEnvironment
 
 void SMassDebugger::RebuildEnvironmentsList()
 {
+	EnvironmentsList.Reset();
 #if WITH_MASSENTITY_DEBUG
 	for (const FMassDebugger::FEnvironment& Environment : FMassDebugger::GetEnvironments())
 	{
@@ -262,6 +277,7 @@ void SMassDebugger::RebuildEnvironmentsList()
 			if (World == nullptr || UE::Mass::Debugger::Private::IsSupportedWorldType(World->WorldType))
 			{
 				EnvironmentsList.Add(MakeShareable(new FMassDebuggerEnvironment(EntityManagerPtr->AsShared())));
+				EnvironmentsList.Last()->ProcessorProviders = Environment.ProcessorProviders;
 			}
 		}
 	}
@@ -306,6 +322,7 @@ void SMassDebugger::BindDelegates()
 #if WITH_MASSENTITY_DEBUG
 	OnEntityManagerInitializedHandle = FMassDebugger::OnEntityManagerInitialized.AddRaw(this, &SMassDebugger::OnEntityManagerInitialized);
 	OnEntityManagerDeinitializedHandle = FMassDebugger::OnEntityManagerDeinitialized.AddRaw(this, &SMassDebugger::OnEntityManagerDeinitialized);
+	OnProcessorProviderRegisteredHandle = FMassDebugger::OnProcessorProviderRegistered.AddRaw(this, &SMassDebugger::OnProcessorProviderRegistered);
 #endif // WITH_MASSENTITY_DEBUG
 }
 
diff --git a/Engine/Source/Editor/MassEntityDebugger/Private/SMassDebugger.h b/Engine/Source/Editor/MassEntityDebugger/Private/SMassDebugger.h
--- a/Engine/Source/Editor/MassEntityDebugger/Private/SMassDebugger.h
+++ b/Engine/Source/Editor/MassEntityDebugger/Private/SMassDebugger.h
@@ -3,10 +3,12 @@
 
 #include "Widgets/Docking/SDockTab.h"
 #include "Framework/Commands/Commands.h"
+#include "MassDebugger.h"
 
 struct FMassDebuggerEnvironment;
 struct FMassDebuggerModel;
 struct FMassEntityManager;
+struct FMassProcessingPhaseManager;
 template<typename T> class SComboBox;
 
 
@@ -47,7 +49,7 @@ public:
 	* Constructs the application.
 	*
 	* @param InArgs The Slate argument list.
-	* @param ConstructUnderMajorTab The major tab which will contain the session front-end.
+	* @param ConstructUnderMajorTab The major tab which will contain the debugger
 	* @param ConstructUnderWindow The window in which this widget is being constructed.
 	*/
 	void Construct(const FArguments& InArgs, const TSharedRef<SDockTab>& ConstructUnderMajorTab, const TSharedPtr<SWindow>& ConstructUnderWindow);
@@ -62,6 +64,7 @@ protected:
 
 	void BindDelegates();
 	void OnEntityManagerInitialized(const FMassEntityManager& EntityManager);
+	void OnProcessorProviderRegistered(const FMassDebugger::FEnvironment& Environment);
 	void OnEntityManagerDeinitialized(const FMassEntityManager& EntityManager);
 	void HandleEnvironmentChanged(TSharedPtr<FMassDebuggerEnvironment> Item, ESelectInfo::Type SelectInfo);
 	void RebuildEnvironmentsList();
@@ -81,4 +84,5 @@ protected:
 
 	FDelegateHandle OnEntityManagerInitializedHandle;
 	FDelegateHandle OnEntityManagerDeinitializedHandle;
+	FDelegateHandle OnProcessorProviderRegisteredHandle;
 };
diff --git a/Engine/Source/Editor/MassEntityDebugger/Private/SMassProcessorsView.cpp b/Engine/Source/Editor/MassEntityDebugger/Private/SMassProcessorsView.cpp
--- a/Engine/Source/Editor/MassEntityDebugger/Private/SMassProcessorsView.cpp
+++ b/Engine/Source/Editor/MassEntityDebugger/Private/SMassProcessorsView.cpp
@@ -2,7 +2,6 @@
 
 #include "SMassProcessorsView.h"
 #include "SMassProcessor.h"
-#include "MassDebuggerModel.h"
 #include "Widgets/Layout/SScrollBox.h"
 #include "Widgets/Views/SListView.h"
 #include "Widgets/SBoxPanel.h"
@@ -18,15 +17,12 @@ namespace UE::Mass::Debug::UI::Private
 //----------------------------------------------------------------------//
 // SMassProcessorListTableRow
 //----------------------------------------------------------------------//
-using FMassDebuggerProcessorDataPtr = TSharedPtr<FMassDebuggerProcessorData, ESPMode::ThreadSafe>;
-
-class SMassProcessorListTableRow : public STableRow<FMassDebuggerProcessorDataPtr>
+class SMassProcessorListTableRow : public STableRow<TSharedPtr<FMassDebuggerProcessorData>>
 {
 public:
-	SLATE_BEGIN_ARGS(SMassProcessorListTableRow) { }
-	SLATE_END_ARGS()
+	using Super = STableRow<TSharedPtr<FMassDebuggerProcessorData>>;
 
-	void Construct(const FArguments& InArgs, const TSharedRef<STableViewBase>& InOwnerTableView, const FMassDebuggerProcessorDataPtr InEntryItem)
+	void Construct(const FArguments& InArgs, const TSharedRef<STableViewBase>& InOwnerTableView, const TSharedPtr<FMassDebuggerProcessorData> InEntryItem)
 	{
 		if (!InEntryItem.IsValid())
 		{
@@ -35,7 +31,7 @@ public:
 
 		Item = InEntryItem;
 		
-		STableRow<FMassDebuggerProcessorDataPtr>::Construct(STableRow<FMassDebuggerProcessorDataPtr>::FArguments(), InOwnerTableView);
+		Super::Construct(Super::FArguments(), InOwnerTableView);
 		
 		ChildSlot
 		[
@@ -45,28 +41,86 @@ public:
 		];
 	}
 
-	FMassDebuggerProcessorDataPtr Item;
+	TSharedPtr<FMassDebuggerProcessorData> Item;
 };
 
-class SMassProcessorTableRow : public STableRow<FMassDebuggerProcessorDataPtr>
+//----------------------------------------------------------------------//
+// SMassProcessorCollectionListView
+//----------------------------------------------------------------------//
+ /**
+  * we have multiple instances of SMassProcessorCollectionListView at one time, and we want selection
+  * cleared on all of them when any of them gets cleared. This class lets us do it by overriding
+  * Private_ClearSelection and notifying the master view about the fact.
+  */
+class SMassProcessorCollectionListView : public SListView<TSharedPtr<FMassDebuggerProcessorData>>
 {
 public:
-	SLATE_BEGIN_ARGS(SMassProcessorTableRow) { }
-	SLATE_END_ARGS()
+	using Super = SListView<TSharedPtr<FMassDebuggerProcessorData>>;
+
+	void Construct(const FArguments& InArgs, const TWeakPtr<SMassProcessorsView>& InWeakMasterView)
+	{
+		WeakMasterView = InWeakMasterView;
+		Super::Construct(InArgs);
+	}
 
-	void Construct(const FArguments& InArgs, const TSharedRef<STableViewBase>& InOwnerTableView, const FMassDebuggerProcessorDataPtr InEntryItem)
+	virtual void Private_ClearSelection() override
+	{
+		if (TSharedPtr<SMassProcessorsView> SharedMasterView = WeakMasterView.Pin())
+		{
+			SharedMasterView->OnClearSelection(*this);
+		}
+		SListView::Private_ClearSelection();
+	}
+
+protected:
+	TWeakPtr<SMassProcessorsView> WeakMasterView;
+};
+
+//----------------------------------------------------------------------//
+// SMassProcessorCollectionTableRow
+//----------------------------------------------------------------------//
+class SMassProcessorCollectionTableRow : public STableRow<TSharedPtr<FMassDebuggerModel::FProcessorCollection>>
+{
+public:
+	using Super = STableRow<TSharedPtr<FMassDebuggerModel::FProcessorCollection>>;
+	using FItemType = TSharedPtr<FMassDebuggerProcessorData>;
+
+	void Construct(const FArguments&, const TSharedRef<STableViewBase>& InOwnerTableView, const TSharedPtr<FMassDebuggerModel::FProcessorCollection> InEntryItem
+		, const TSharedRef<SMassProcessorsView>& MasterView)
 	{
 		Item = InEntryItem;
 
-		STableRow<FMassDebuggerProcessorDataPtr>::Construct(STableRow<FMassDebuggerProcessorDataPtr>::FArguments(), InOwnerTableView);
+		Super::Construct(Super::FArguments(), InOwnerTableView);
+
+		TSharedPtr<SMassProcessorCollectionListView>& ProcessorsListWidget = MasterView->ProcessorsListWidgets.AddDefaulted_GetRef();
 
 		ChildSlot
 		[
-			SNew(SMassProcessor, Item)
+			SNew(SVerticalBox)
+
+			+ SVerticalBox::Slot()
+			[
+				SNew(STextBlock)
+				.Text(FText::FromName(Item->Label))
+				.TextStyle(FAppStyle::Get(), "LargeText")
+			]
+			+ SVerticalBox::Slot()
+			.AutoHeight()
+			[
+				SAssignNew(ProcessorsListWidget, SMassProcessorCollectionListView, MasterView)
+				.ListItemsSource(&InEntryItem->Container)
+				.SelectionMode(ESelectionMode::Multi)
+				.OnSelectionChanged(MasterView->OnProcessorSelectionChanged)
+				.OnGenerateRow_Lambda([](FItemType Item, const TSharedRef<STableViewBase>& OwnerTable)
+					{
+						return SNew(SMassProcessorListTableRow, OwnerTable, Item);
+					})
+			]
 		];
 	}
 
-	FMassDebuggerProcessorDataPtr Item;
+protected:
+	TSharedPtr<FMassDebuggerModel::FProcessorCollection> Item;
 };
 
 //----------------------------------------------------------------------//
@@ -85,23 +139,30 @@ void SMassProcessorsView::Construct(const FArguments& InArgs, TSharedRef<FMassDe
 		];
 	}
 
+	OnProcessorSelectionChanged = FOnProcessorSelectionChanged::CreateSP(this, &SMassProcessorsView::ProcessorListSelectionChanged);
+	
 	ChildSlot
 		[
 			SNew(SSplitter)
 			.Orientation(Orient_Horizontal)
 			 
 			+ SSplitter::Slot()
-			.Value(.3f)
-			.MinSize(250.0f)
+			.Value(.35f)
+			.MinSize(260.0f)
 			[
-				SAssignNew(ProcessorsListWidget, SListView<TSharedPtr<FMassDebuggerProcessorData> >)
-					.ListItemsSource(&DebuggerModel->CachedProcessors)
-					.SelectionMode(ESelectionMode::Multi)
-					.OnSelectionChanged(this, &SMassProcessorsView::ProcessorListSelectionChanged)
-					.OnGenerateRow_Lambda([](TSharedPtr<FMassDebuggerProcessorData> Item, const TSharedRef<STableViewBase>& OwnerTable)
+				SNew(SScrollBox)
+				.Orientation(Orient_Vertical)
+				+ SScrollBox::Slot()
+				[
+					SAssignNew(ProcessorCollectionsListWidget, SListView<TSharedPtr<FMassDebuggerModel::FProcessorCollection>>)
+					.ListItemsSource(&DebuggerModel->CachedProcessorCollections)
+					.SelectionMode(ESelectionMode::None)
+					.OnGenerateRow_Lambda([SharedThis = StaticCastSharedRef<SMassProcessorsView>(AsShared())]
+						(TSharedPtr<FMassDebuggerModel::FProcessorCollection> Item, const TSharedRef<STableViewBase>& OwnerTable)
 						{
-							return SNew(SMassProcessorListTableRow, OwnerTable, Item);
+							return SNew(SMassProcessorCollectionTableRow, OwnerTable, Item, SharedThis);
 						})
+				]
 			]
 			
 			+ SSplitter::Slot()
@@ -130,10 +191,33 @@ void SMassProcessorsView::ProcessorListSelectionChanged(TSharedPtr<FMassDebugger
 	check(DebuggerModel);
 
 	TArray<TSharedPtr<FMassDebuggerProcessorData>> CurrentlySelectedProcessors;
-	ProcessorsListWidget->GetSelectedItems(CurrentlySelectedProcessors);
+	for (const auto& Widget : ProcessorsListWidgets)
+	{
+		TArray<TSharedPtr<FMassDebuggerProcessorData>> LocalSelectedProcessors;
+		Widget->GetSelectedItems(LocalSelectedProcessors);
+		CurrentlySelectedProcessors.Append(MoveTemp(LocalSelectedProcessors));
+	}
 	DebuggerModel->SelectProcessors(CurrentlySelectedProcessors, SelectInfo);
 }
 
+void SMassProcessorsView::OnClearSelection(const SMassProcessorCollectionListView& TransientSource)
+{
+	if (bClearingSelection)
+	{
+		return;
+	}
+
+	TGuardValue<bool> GuardValue(bClearingSelection, true);
+	
+	for (const auto& Widget : ProcessorsListWidgets)
+	{
+		if (Widget.Get() != &TransientSource)
+		{
+			Widget->ClearSelection();
+		}
+	}
+}
+
 void SMassProcessorsView::OnProcessorsSelected(TConstArrayView<TSharedPtr<FMassDebuggerProcessorData>> SelectedProcessors, ESelectInfo::Type SelectInfo)
 {
 	using namespace UE::Mass::Debug::UI::Private;
@@ -146,7 +230,10 @@ void SMassProcessorsView::OnProcessorsSelected(TConstArrayView<TSharedPtr<FMassD
 
 	if (SelectInfo == ESelectInfo::Direct)
 	{
-		ProcessorsListWidget->ClearSelection();
+		for (const auto& Widget : ProcessorsListWidgets)
+		{
+			Widget->ClearSelection();
+		}
 	}
 	ProcessorsBox->ClearChildren();
 
@@ -154,9 +241,10 @@ void SMassProcessorsView::OnProcessorsSelected(TConstArrayView<TSharedPtr<FMassD
 	{		
 		if (SelectInfo == ESelectInfo::Direct)
 		{
-			ProcessorsListWidget->SetItemSelection(SelectedProcessors, /*bSelected=*/true, ESelectInfo::OnMouseClick);
-			// scroll to the first item to make sure there's anything in view
-			ProcessorsListWidget->RequestScrollIntoView(SelectedProcessors[0]);
+			for (const auto& Widget : ProcessorsListWidgets)
+			{
+				Widget->SetItemSelection(SelectedProcessors, /*bSelected=*/true, ESelectInfo::OnMouseClick);
+			}
 		}
 
 		for (TSharedPtr<FMassDebuggerProcessorData>& ProcessorData : DebuggerModel->SelectedProcessors)
@@ -168,7 +256,8 @@ void SMassProcessorsView::OnProcessorsSelected(TConstArrayView<TSharedPtr<FMassD
 				];
 		}
 	}
-	ProcessorsListWidget->RequestListRefresh();
+
+	ProcessorCollectionsListWidget->RequestListRefresh();
 }
 
 void SMassProcessorsView::OnArchetypesSelected(TConstArrayView<TSharedPtr<FMassDebuggerArchetypeData>> SelectedArchetypes, ESelectInfo::Type SelectInfo)
@@ -184,10 +273,9 @@ void SMassProcessorsView::OnRefresh()
 void SMassProcessorsView::PopulateProcessorList()
 {
 	check(DebuggerModel.IsValid());
-	ProcessorsListWidget->ClearSelection();
 	DebuggerModel->ClearProcessorSelection();
 
-	ProcessorsListWidget->RequestListRefresh();
+	ProcessorCollectionsListWidget->RequestListRefresh();
 
 	if (ProcessorsBox)
 	{
diff --git a/Engine/Source/Editor/MassEntityDebugger/Private/SMassProcessorsView.h b/Engine/Source/Editor/MassEntityDebugger/Private/SMassProcessorsView.h
--- a/Engine/Source/Editor/MassEntityDebugger/Private/SMassProcessorsView.h
+++ b/Engine/Source/Editor/MassEntityDebugger/Private/SMassProcessorsView.h
@@ -2,13 +2,15 @@
 #pragma once
 
 #include "SMassDebuggerViewBase.h"
-
+#include "MassDebuggerModel.h"
+#include "Templates/SharedPointer.h"
 
 struct FMassDebuggerProcessorData;
 struct FMassDebuggerArchetypeData;
 class SRichTextBlock;
-struct FMassDebuggerModel;
 class SMassProcessor;
+class SMassProcessorCollectionListView;
+class SMassProcessorCollectionTableRow;
 template<typename T> class SListView;
 class SVerticalBox;
 
@@ -31,8 +33,17 @@ protected:
 
 	void ProcessorListSelectionChanged(TSharedPtr<FMassDebuggerProcessorData> SelectedItem, ESelectInfo::Type SelectInfo);
 
-	TSharedPtr<SListView<TSharedPtr<FMassDebuggerProcessorData>>> ProcessorsListWidget;	
-	TSharedPtr<SListView<TSharedPtr<FMassDebuggerProcessorData>>> SelectedProcessorsListWidget;
+	friend SMassProcessorCollectionListView;
+	friend SMassProcessorCollectionTableRow;
+	void OnClearSelection(const SMassProcessorCollectionListView& TransientSource);
+
+	TSharedPtr<SListView<TSharedPtr<FMassDebuggerModel::FProcessorCollection>>> ProcessorCollectionsListWidget;
+	TArray<TSharedPtr<SMassProcessorCollectionListView>> ProcessorsListWidgets;
 	TSharedPtr<SMassProcessor> ProcessorWidget;
 	TSharedPtr<SVerticalBox> ProcessorsBox;
+
+	DECLARE_DELEGATE_TwoParams(FOnProcessorSelectionChanged, TSharedPtr<FMassDebuggerProcessorData>, ESelectInfo::Type);
+	FOnProcessorSelectionChanged OnProcessorSelectionChanged;
+
+	bool bClearingSelection = false;
 };
diff --git a/Engine/Source/Runtime/MassEntity/Private/MassDebugger.cpp b/Engine/Source/Runtime/MassEntity/Private/MassDebugger.cpp
--- a/Engine/Source/Runtime/MassEntity/Private/MassDebugger.cpp
+++ b/Engine/Source/Runtime/MassEntity/Private/MassDebugger.cpp
@@ -377,6 +377,7 @@ FMassDebugger::FOnEntitySelected FMassDebugger::OnEntitySelectedDelegate;
 
 FMassDebugger::FOnMassEntityManagerEvent FMassDebugger::OnEntityManagerInitialized;
 FMassDebugger::FOnMassEntityManagerEvent FMassDebugger::OnEntityManagerDeinitialized;
+FMassDebugger::FOnEnvironmentEvent FMassDebugger::OnProcessorProviderRegistered;
 FMassDebugger::FOnDebugEvent FMassDebugger::OnDebugEvent;
 TArray<FMassDebugger::FEnvironment> FMassDebugger::ActiveEnvironments;
 UE::FSpinLock FMassDebugger::EntityManagerRegistrationLock;
@@ -682,20 +683,22 @@ FMassEntityHandle FMassDebugger::GetSelectedEntity(const FMassEntityManager& Ent
 	return Index != INDEX_NONE ? ActiveEnvironments[Index].SelectedEntity : FMassEntityHandle();
 }
 
-void FMassDebugger::RegisterEntityManager(FMassEntityManager& EntityManager)
+int32 FMassDebugger::RegisterEntityManager(FMassEntityManager& EntityManager)
 {
-	UE::TScopeLock<UE::FSpinLock> ScopeLock(EntityManagerRegistrationLock);
-
-	ActiveEnvironments.Emplace(EntityManager);
+	int32 NewEnvironmentIndex = INDEX_NONE;
+	{
+		UE::TScopeLock<UE::FSpinLock> ScopeLock(EntityManagerRegistrationLock);
+		NewEnvironmentIndex = ActiveEnvironments.Emplace(EntityManager);
+	}
 	OnEntityManagerInitialized.Broadcast(EntityManager);
+	return NewEnvironmentIndex;
 }
 
 void FMassDebugger::UnregisterEntityManager(FMassEntityManager& EntityManager)
 {
-	UE::TScopeLock<UE::FSpinLock> ScopeLock(EntityManagerRegistrationLock);
-
 	if (EntityManager.DoesSharedInstanceExist())
 	{
+		UE::TScopeLock<UE::FSpinLock> ScopeLock(EntityManagerRegistrationLock);
 		const int32 Index = ActiveEnvironments.IndexOfByPredicate([WeakManager = EntityManager.AsWeak()](const FEnvironment& Element) 
 		{
 			return Element.EntityManager == WeakManager;
@@ -707,6 +710,7 @@ void FMassDebugger::UnregisterEntityManager(FMassEntityManager& EntityManager)
 	}
 	else
 	{
+		UE::TScopeLock<UE::FSpinLock> ScopeLock(EntityManagerRegistrationLock);
 		ActiveEnvironments.RemoveAll([](const FEnvironment& Item)
 			{
 				return Item.IsValid() == false;
@@ -715,6 +719,24 @@ void FMassDebugger::UnregisterEntityManager(FMassEntityManager& EntityManager)
 	OnEntityManagerDeinitialized.Broadcast(EntityManager);
 }
 
+void FMassDebugger::RegisterProcessorDataProvider(FName ProviderName, const TSharedRef<FMassEntityManager>& EntityManager, const FProcessorProviderFunction& ProviderFunction)
+{
+	UE::TScopeLock<UE::FSpinLock> ScopeLock(EntityManagerRegistrationLock);
+	int32 Index = ActiveEnvironments.IndexOfByPredicate([WeakEntityManager = EntityManager->AsWeak()](const FEnvironment& Element) 
+	{
+		return Element.EntityManager == WeakEntityManager;
+	});
+
+	if (Index == INDEX_NONE)
+	{
+		Index = RegisterEntityManager(*EntityManager);
+	}
+	
+	ActiveEnvironments[Index].ProcessorProviders.FindOrAdd(ProviderName, ProviderFunction);
+
+	OnProcessorProviderRegistered.Broadcast(ActiveEnvironments[Index]);
+}
+
 bool FMassDebugger::DoesArchetypeMatchRequirements(const FMassArchetypeHandle& ArchetypeHandle, const FMassFragmentRequirements& Requirements, FOutputDevice& OutputDevice)
 {
 	if (const FMassArchetypeData* Archetype = FMassArchetypeHelper::ArchetypeDataFromHandle(ArchetypeHandle))
diff --git a/Engine/Source/Runtime/MassEntity/Private/MassObserverManager.cpp b/Engine/Source/Runtime/MassEntity/Private/MassObserverManager.cpp
--- a/Engine/Source/Runtime/MassEntity/Private/MassObserverManager.cpp
+++ b/Engine/Source/Runtime/MassEntity/Private/MassObserverManager.cpp
@@ -6,6 +6,7 @@
 #include "MassExecutor.h"
 #include "MassProcessingTypes.h"
 #include "MassObserverRegistry.h"
+#include "MassDebugger.h"
 
 #include UE_INLINE_GENERATED_CPP_BY_NAME(MassObserverManager)
 
@@ -49,6 +50,23 @@ namespace UE::Mass::ObserverManager
 	} // Private
 } // UE::Mass::ObserverManager
 
+//----------------------------------------------------------------------//
+// FMassObserversMap
+//----------------------------------------------------------------------//
+void FMassObserversMap::DebugAddUniqueProcessors(TArray<const UMassProcessor*>& OutProcessors) const
+{
+#if WITH_MASSENTITY_DEBUG
+	for (const auto& MapElement : Container)
+	{
+		for (const UMassProcessor* Processor : MapElement.Value.GetProcessorsView())
+		{
+			ensure(Processor);
+			OutProcessors.AddUnique(Processor);
+		}
+	}
+#endif // WITH_MASSENTITY_DEBUG
+}
+
 //----------------------------------------------------------------------//
 // FMassObserverManager
 //----------------------------------------------------------------------//
@@ -64,6 +82,17 @@ FMassObserverManager::FMassObserverManager(FMassEntityManager& Owner)
 
 }
 
+void FMassObserverManager::DebugGatherUniqueProcessors(TArray<const UMassProcessor*>& OutProcessors) const
+{
+#if WITH_MASSENTITY_DEBUG
+	for (int32 OperationIndex = 0; OperationIndex < static_cast<uint32>(EMassObservedOperation::MAX); ++OperationIndex)
+	{
+		FragmentObservers[OperationIndex].DebugAddUniqueProcessors(OutProcessors);
+		TagObservers[OperationIndex].DebugAddUniqueProcessors(OutProcessors);
+	}
+#endif // WITH_MASSENTITY_DEBUG
+}
+
 void FMassObserverManager::Initialize()
 {
 	// instantiate initializers
@@ -80,6 +109,18 @@ void FMassObserverManager::Initialize()
 		SetUpObservers(EntityManager, WorldExecutionFlags, *Owner, *Registry.FragmentObservers[i], ObservedFragments[i], FragmentObservers[i]);
 		SetUpObservers(EntityManager, WorldExecutionFlags, *Owner, *Registry.TagObservers[i], ObservedTags[i], TagObservers[i]);
 	}
+
+#if WITH_MASSENTITY_DEBUG
+	FMassDebugger::RegisterProcessorDataProvider(TEXT("Observers"), EntityManager.AsShared()
+		, [WeakManager = EntityManager.AsWeak()](TArray<const UMassProcessor*>& OutProcessors)
+	{
+		if (TSharedPtr<FMassEntityManager> EntityManager = WeakManager.Pin())
+		{
+			FMassObserverManager& ObserverManager = EntityManager->GetObserverManager();
+			ObserverManager.DebugGatherUniqueProcessors(OutProcessors);
+		}
+	});
+#endif // WITH_MASSENTITY_DEBUG
 }
 
 void FMassObserverManager::DeInitialize()
diff --git a/Engine/Source/Runtime/MassEntity/Private/MassProcessingPhaseManager.cpp b/Engine/Source/Runtime/MassEntity/Private/MassProcessingPhaseManager.cpp
--- a/Engine/Source/Runtime/MassEntity/Private/MassProcessingPhaseManager.cpp
+++ b/Engine/Source/Runtime/MassEntity/Private/MassProcessingPhaseManager.cpp
@@ -155,8 +155,10 @@ void FMassPhaseProcessorConfigurationHelper::Configure(TArrayView<UMassProcessor
 
 		for (UMassProcessor* PrunedProcessor : InOutOptionalResult.PrunedProcessors)
 		{
-			if (PrunedProcessor)
+			if (ensure(PrunedProcessor))
 			{
+				// @todo consider adding all PrunedProcessors in one go - there should be no overlap
+				// with PhaseProcessor.GetChildProcessorsView(), not there should be any nulls.
 				TmpPipeline.AppendProcessor(*PrunedProcessor);
 			}
 		}
@@ -226,6 +228,15 @@ void FMassPhaseProcessorConfigurationHelper::Configure(TArrayView<UMassProcessor
 	}
 }
 
+//----------------------------------------------------------------------//
+// FMassProcessingPhaseManager::FPhaseGraphBuildState
+//----------------------------------------------------------------------//
+void FMassProcessingPhaseManager::FPhaseGraphBuildState::Reset()
+{
+	LastResult.Reset();
+	bInitialized = false;
+}
+
 //----------------------------------------------------------------------//
 // FMassProcessingPhaseManager
 //----------------------------------------------------------------------//
@@ -285,7 +296,7 @@ void FMassProcessingPhaseManager::Deinitialize()
 
 	for (FPhaseGraphBuildState& GraphBuildState : ProcessingGraphBuildStates)
 	{
-		GraphBuildState.LastResult.PrunedProcessors.Reset();
+		GraphBuildState.Reset();
 	}
 }
 
@@ -315,13 +326,38 @@ void FMassProcessingPhaseManager::Start(UWorld& World)
 	}
 }
 
-void FMassProcessingPhaseManager::Start(const TSharedPtr<FMassEntityManager>& InEntityManager)
+void FMassProcessingPhaseManager::Start(const TSharedRef<FMassEntityManager>& InEntityManager)
 {
 	EntityManager = InEntityManager;
 
+#if WITH_MASSENTITY_DEBUG
+	FMassDebugger::RegisterProcessorDataProvider(TEXT("Phase-executed processors"), InEntityManager, [WeakThis = AsWeak()](TArray<const UMassProcessor*>& OutProcessors)
+	{
+		if (TSharedPtr<FMassProcessingPhaseManager> SharedThis = WeakThis.Pin())
+		{
+			for (const FMassProcessingPhase& Phase : SharedThis->ProcessingPhases)
+			{
+				OutProcessors.Append(Phase.DebugGetPhaseProcessor()->GetChildProcessorsView());
+			}
+		}
+	});
+
+	FMassDebugger::RegisterProcessorDataProvider(TEXT("Pruned processors"), InEntityManager, [WeakThis = AsWeak()](TArray<const UMassProcessor*>& OutProcessors)
+	{
+		if (TSharedPtr<FMassProcessingPhaseManager> SharedThis = WeakThis.Pin())
+		{
+			TConstArrayView<FPhaseGraphBuildState> BuildStates = SharedThis->DebugGetProcessingGraphBuildStates();
+			for (const FPhaseGraphBuildState& State : BuildStates)
+			{
+				OutProcessors.Append(ObjectPtrDecay(State.LastResult.PrunedProcessors));
+			}
+		}
+	});
+#endif // WITH_MASSENTITY_DEBUG
+
 	OnNewArchetypeHandle = EntityManager->GetOnNewArchetypeEvent().AddRaw(this, &FMassProcessingPhaseManager::OnNewArchetype);
 
-	if (UWorld* World = InEntityManager->GetWorld())
+	if (UWorld* World = EntityManager->GetWorld())
 	{
 		EnableTickFunctions(*World);
 	}
@@ -554,4 +590,12 @@ void FMassPhaseProcessorConfigurationHelper::Configure(TArrayView<UMassProcessor
 	}
 }
 
+void FMassProcessingPhaseManager::Start(const TSharedPtr<FMassEntityManager>& InEntityManager)
+{
+	if (InEntityManager)
+	{
+		Start(InEntityManager.ToSharedRef());
+	}
+}
+
 #undef LOCTEXT_NAMESPACE
diff --git a/Engine/Source/Runtime/MassEntity/Public/MassDebugger.h b/Engine/Source/Runtime/MassEntity/Public/MassDebugger.h
--- a/Engine/Source/Runtime/MassEntity/Public/MassDebugger.h
+++ b/Engine/Source/Runtime/MassEntity/Public/MassDebugger.h
@@ -132,13 +132,12 @@ namespace UE::Mass::Debug
 
 struct MASSENTITY_API FMassDebugger
 {
-	DECLARE_MULTICAST_DELEGATE_TwoParams(FOnEntitySelected, const FMassEntityManager&, const FMassEntityHandle);
-	DECLARE_MULTICAST_DELEGATE_OneParam(FOnMassEntityManagerEvent, const FMassEntityManager&);
-	DECLARE_MULTICAST_DELEGATE_ThreeParams(FOnDebugEvent, const FName /*EventName*/, FConstStructView /*Payload*/, const EMassDebugMessageSeverity /*SeverityOverride*/);
+	using FProcessorProviderFunction = TFunction<void(TArray<const UMassProcessor*>&)>;
 
 	struct FEnvironment
 	{
 		TWeakPtr<const FMassEntityManager> EntityManager;
+		TMap<FName, FProcessorProviderFunction> ProcessorProviders;
 		FMassEntityHandle SelectedEntity;
 
 		explicit FEnvironment(const FMassEntityManager& InEntityManager)
@@ -148,6 +147,11 @@ struct MASSENTITY_API FMassDebugger
 		bool IsValid() const { return EntityManager.IsValid(); }
 	};
 
+	DECLARE_TS_MULTICAST_DELEGATE_TwoParams(FOnEntitySelected, const FMassEntityManager&, const FMassEntityHandle);
+	DECLARE_TS_MULTICAST_DELEGATE_OneParam(FOnMassEntityManagerEvent, const FMassEntityManager&);
+	DECLARE_TS_MULTICAST_DELEGATE_OneParam(FOnEnvironmentEvent, const FEnvironment&);
+	DECLARE_TS_MULTICAST_DELEGATE_ThreeParams(FOnDebugEvent, const FName /*EventName*/, FConstStructView /*Payload*/, const EMassDebugMessageSeverity /*SeverityOverride*/);
+
 	static TConstArrayView<FMassEntityQuery*> GetProcessorQueries(const UMassProcessor& Processor);
 	/** fetches all queries registered for given Processor. Note that in order to get up to date information
 	 *  FMassEntityQuery::CacheArchetypes will be called on each query */
@@ -181,6 +185,7 @@ struct MASSENTITY_API FMassDebugger
 
 	static FOnMassEntityManagerEvent OnEntityManagerInitialized;
 	static FOnMassEntityManagerEvent OnEntityManagerDeinitialized;
+	static FOnEnvironmentEvent OnProcessorProviderRegistered;
 
 	static FOnDebugEvent OnDebugEvent;
 	
@@ -196,8 +201,21 @@ struct MASSENTITY_API FMassDebugger
 			, FConstStructView::Make(TMessage(Forward<TArgs>(InArgs)...)));
 	}
 
-	static void RegisterEntityManager(FMassEntityManager& EntityManager);
+	/**
+	 * Registers given EntityManager with the debugger, creating a new entry in ActiveEnvironments.
+	 * @return the index of newly created environment
+	 */
+	static int32 RegisterEntityManager(FMassEntityManager& EntityManager);
 	static void UnregisterEntityManager(FMassEntityManager& EntityManager);
+	/**
+	 * Registers the given ProviderFunction with the existing FEnvironment associated with the provided EntityManager.
+	 * If one doesn't exist yet, it will be created (i.e. will automatically call RegisterEntityManager() with the provided EntityManager).
+	 * The function will be called during data collection for the given FEnvironment.
+	 * NOTE: there's no UnregisterProcessorDataProvider, the registered providers will automatically get removed along with
+	 * the rest of the data associated with the relevant EntityManager as part of UnregisterEntityManager call.
+	 */
+	static void RegisterProcessorDataProvider(FName ProviderName, const TSharedRef<FMassEntityManager>& EntityManager, const FProcessorProviderFunction& ProviderFunction);
+
 	static TConstArrayView<FEnvironment> GetEnvironments() { return ActiveEnvironments; }
 
 	/**
diff --git a/Engine/Source/Runtime/MassEntity/Public/MassObserverManager.h b/Engine/Source/Runtime/MassEntity/Public/MassObserverManager.h
--- a/Engine/Source/Runtime/MassEntity/Public/MassObserverManager.h
+++ b/Engine/Source/Runtime/MassEntity/Public/MassObserverManager.h
@@ -27,6 +27,8 @@ struct FMassObserversMap
 		return Container;
 	}
 
+	void DebugAddUniqueProcessors(TArray<const UMassProcessor*>& OutProcessors) const;
+
 private:
 	UPROPERTY()
 	TMap<TObjectPtr<const UScriptStruct>, FMassRuntimePipeline> Container;
@@ -129,6 +131,8 @@ public:
 	void AddObserverInstance(const UScriptStruct& FragmentOrTagType, const EMassObservedOperation Operation, UMassProcessor& ObserverProcessor);
 	void RemoveObserverInstance(const UScriptStruct& FragmentOrTagType, const EMassObservedOperation Operation, UMassProcessor& ObserverProcessor);
 
+	void DebugGatherUniqueProcessors(TArray<const UMassProcessor*>& OutProcessors) const;
+
 protected:
 	friend FMassEntityManager;
 	explicit FMassObserverManager(FMassEntityManager& Owner);
diff --git a/Engine/Source/Runtime/MassEntity/Public/MassProcessingPhaseManager.h b/Engine/Source/Runtime/MassEntity/Public/MassProcessingPhaseManager.h
--- a/Engine/Source/Runtime/MassEntity/Public/MassProcessingPhaseManager.h
+++ b/Engine/Source/Runtime/MassEntity/Public/MassProcessingPhaseManager.h
@@ -77,6 +77,10 @@ public:
 	void AddSupportedTickType(const ELevelTick TickType) { SupportedTickTypes |= (1 << TickType); }
 	void RemoveSupportedTickType(const ELevelTick TickType) { SupportedTickTypes &= ~(1 << TickType); }
 
+#if WITH_MASSENTITY_DEBUG
+	const UMassCompositeProcessor* DebugGetPhaseProcessor() const;
+#endif // WITH_MASSENTITY_DEBUG
+
 protected:
 	friend FMassProcessingPhaseManager;
 
@@ -167,7 +171,7 @@ public:
 	/**
 	 *  Stores InEntityManager as the entity manager. It also kicks off phase ticking if the given InEntityManager is tied to a UWorld.
 	 */
-	void Start(const TSharedPtr<FMassEntityManager>& InEntityManager);
+	void Start(const TSharedRef<FMassEntityManager>& InEntityManager);
 	void Stop();
 	bool IsRunning() const { return EntityManager.IsValid(); }
 
@@ -178,6 +182,21 @@ public:
 	/** Removes a previously registered dynamic processor of throws an assert if not found. */
 	void UnregisterDynamicProcessor(UMassProcessor& Processor);
 
+	struct FPhaseGraphBuildState
+	{
+		FMassProcessorDependencySolver::FResult LastResult;
+		bool bNewArchetypes = true;
+		bool bProcessorsNeedRebuild = true;
+		bool bInitialized = false;
+
+		void Reset();
+	};
+
+#if WITH_MASSENTITY_DEBUG
+	TConstArrayView<FMassProcessingPhase>  DebugGetProcessingPhases() const;
+	TConstArrayView<FPhaseGraphBuildState> DebugGetProcessingGraphBuildStates() const;
+#endif // WITH_MASSENTITY_DEBUG
+
 protected:
 	// FGCObject interface
 	virtual void AddReferencedObjects(FReferenceCollector& Collector) override;
@@ -209,15 +228,7 @@ protected:
 
 	void OnNewArchetype(const FMassArchetypeHandle& NewArchetype);
 
-protected:	
-	struct FPhaseGraphBuildState
-	{
-		FMassProcessorDependencySolver::FResult LastResult;
-		bool bNewArchetypes = true;
-		bool bProcessorsNeedRebuild = true;
-		bool bInitialized = false;
-	};
-
+protected:
 	FMassProcessingPhase ProcessingPhases[(uint8)EMassProcessingPhase::MAX];
 	FPhaseGraphBuildState ProcessingGraphBuildStates[(uint8)EMassProcessingPhase::MAX];
 	TArray<FMassProcessingPhaseConfig> ProcessingPhasesConfig;
@@ -242,4 +253,29 @@ protected:
 	void OnDebugEntityManagerInitialized(const FMassEntityManager&);
 	void OnDebugEntityManagerDeinitialized(const FMassEntityManager&);
 #endif // WITH_MASSENTITY_DEBUG
+
+public:
+	UE_DEPRECATED(5.6, "This flavor of Start is deprecated. Please use the one using a TSharedRef<FMassEntityManager> parameter instead")
+	void Start(const TSharedPtr<FMassEntityManager>& InEntityManager);
 };
+
+//----------------------------------------------------------------------//
+// inlines
+//----------------------------------------------------------------------//
+#if WITH_MASSENTITY_DEBUG
+inline const UMassCompositeProcessor* FMassProcessingPhase::DebugGetPhaseProcessor() const
+{
+	return PhaseProcessor;
+}
+
+inline TConstArrayView<FMassProcessingPhase>  FMassProcessingPhaseManager::DebugGetProcessingPhases() const
+{
+	return MakeArrayView(ProcessingPhases, static_cast<uint8>(EMassProcessingPhase::MAX));
+}
+
+inline TConstArrayView<FMassProcessingPhaseManager::FPhaseGraphBuildState> FMassProcessingPhaseManager::DebugGetProcessingGraphBuildStates() const
+{
+	return MakeArrayView(ProcessingGraphBuildStates, static_cast<uint8>(EMassProcessingPhase::MAX));
+}
+#endif // WITH_MASSENTITY_DEBUG
+
