[Mass]
CIS Issue 640897 - Compilation fixes to FMassObserverManager and SMassProcessorView
CIS Issue 640915

diff --git a/Engine/Source/Editor/MassEntityDebugger/Private/SMassProcessorsView.cpp b/Engine/Source/Editor/MassEntityDebugger/Private/SMassProcessorsView.cpp
--- a/Engine/Source/Editor/MassEntityDebugger/Private/SMassProcessorsView.cpp
+++ b/Engine/Source/Editor/MassEntityDebugger/Private/SMassProcessorsView.cpp
@@ -22,7 +22,7 @@ class SMassProcessorListTableRow : public STableRow<TSharedPtr<FMassDebuggerProc
 public:
 	using Super = STableRow<TSharedPtr<FMassDebuggerProcessorData>>;
 
-	void Construct(const FArguments& InArgs, const TSharedRef<STableViewBase>& InOwnerTableView, const TSharedPtr<FMassDebuggerProcessorData> InEntryItem)
+	void Construct(const FArguments& InArgs, const TSharedRef<STableViewBase>& InOwnerTableView, const TSharedPtr<FMassDebuggerProcessorData>& InEntryItem)
 	{
 		if (!InEntryItem.IsValid())
 		{
@@ -50,30 +50,30 @@ public:
  /**
   * we have multiple instances of SMassProcessorCollectionListView at one time, and we want selection
   * cleared on all of them when any of them gets cleared. This class lets us do it by overriding
-  * Private_ClearSelection and notifying the master view about the fact.
+  * Private_ClearSelection and notifying the main view about the fact.
   */
 class SMassProcessorCollectionListView : public SListView<TSharedPtr<FMassDebuggerProcessorData>>
 {
 public:
 	using Super = SListView<TSharedPtr<FMassDebuggerProcessorData>>;
 
-	void Construct(const FArguments& InArgs, const TWeakPtr<SMassProcessorsView>& InWeakMasterView)
+	void Construct(const FArguments& InArgs, const TWeakPtr<SMassProcessorsView>& InWeakMainView)
 	{
-		WeakMasterView = InWeakMasterView;
+		WeakMainView = InWeakMainView;
 		Super::Construct(InArgs);
 	}
 
 	virtual void Private_ClearSelection() override
 	{
-		if (TSharedPtr<SMassProcessorsView> SharedMasterView = WeakMasterView.Pin())
+		if (TSharedPtr<SMassProcessorsView> SharedMainView = WeakMainView.Pin())
 		{
-			SharedMasterView->OnClearSelection(*this);
+			SharedMainView->OnClearSelection(*this);
 		}
 		SListView::Private_ClearSelection();
 	}
 
 protected:
-	TWeakPtr<SMassProcessorsView> WeakMasterView;
+	TWeakPtr<SMassProcessorsView> WeakMainView;
 };
 
 //----------------------------------------------------------------------//
@@ -85,14 +85,14 @@ public:
 	using Super = STableRow<TSharedPtr<FMassDebuggerModel::FProcessorCollection>>;
 	using FItemType = TSharedPtr<FMassDebuggerProcessorData>;
 
-	void Construct(const FArguments&, const TSharedRef<STableViewBase>& InOwnerTableView, const TSharedPtr<FMassDebuggerModel::FProcessorCollection> InEntryItem
-		, const TSharedRef<SMassProcessorsView>& MasterView)
+	void Construct(const FArguments&, const TSharedRef<STableViewBase>& InOwnerTableView, const TSharedPtr<FMassDebuggerModel::FProcessorCollection>& InEntryItem
+		, const TSharedRef<SMassProcessorsView>& MainView)
 	{
-		Item = InEntryItem;
+		CollectionItem = InEntryItem;
 
 		Super::Construct(Super::FArguments(), InOwnerTableView);
 
-		TSharedPtr<SMassProcessorCollectionListView>& ProcessorsListWidget = MasterView->ProcessorsListWidgets.AddDefaulted_GetRef();
+		TSharedPtr<SMassProcessorCollectionListView>& ProcessorsListWidget = MainView->ProcessorsListWidgets.AddDefaulted_GetRef();
 
 		ChildSlot
 		[
@@ -101,16 +101,16 @@ public:
 			+ SVerticalBox::Slot()
 			[
 				SNew(STextBlock)
-				.Text(FText::FromName(Item->Label))
+				.Text(FText::FromName(CollectionItem->Label))
 				.TextStyle(FAppStyle::Get(), "LargeText")
 			]
 			+ SVerticalBox::Slot()
 			.AutoHeight()
 			[
-				SAssignNew(ProcessorsListWidget, SMassProcessorCollectionListView, MasterView)
-				.ListItemsSource(&InEntryItem->Container)
+				SAssignNew(ProcessorsListWidget, SMassProcessorCollectionListView, MainView)
+				.ListItemsSource(&CollectionItem->Container)
 				.SelectionMode(ESelectionMode::Multi)
-				.OnSelectionChanged(MasterView->OnProcessorSelectionChanged)
+				.OnSelectionChanged(MainView->OnProcessorSelectionChanged)
 				.OnGenerateRow_Lambda([](FItemType Item, const TSharedRef<STableViewBase>& OwnerTable)
 					{
 						return SNew(SMassProcessorListTableRow, OwnerTable, Item);
@@ -120,13 +120,13 @@ public:
 	}
 
 protected:
-	TSharedPtr<FMassDebuggerModel::FProcessorCollection> Item;
+	TSharedPtr<FMassDebuggerModel::FProcessorCollection> CollectionItem;
 };
 
 //----------------------------------------------------------------------//
 // SMassProcessorsView
 //----------------------------------------------------------------------//
-void SMassProcessorsView::Construct(const FArguments& InArgs, TSharedRef<FMassDebuggerModel> InDebuggerModel)
+void SMassProcessorsView::Construct(const FArguments& InArgs, const TSharedRef<FMassDebuggerModel>& InDebuggerModel)
 {
 	Initialize(InDebuggerModel);
 
diff --git a/Engine/Source/Editor/MassEntityDebugger/Private/SMassProcessorsView.h b/Engine/Source/Editor/MassEntityDebugger/Private/SMassProcessorsView.h
--- a/Engine/Source/Editor/MassEntityDebugger/Private/SMassProcessorsView.h
+++ b/Engine/Source/Editor/MassEntityDebugger/Private/SMassProcessorsView.h
@@ -22,7 +22,7 @@ public:
 	}
 	SLATE_END_ARGS()
 
-	void Construct(const FArguments& InArgs, TSharedRef<FMassDebuggerModel> InDebuggerModel);
+	void Construct(const FArguments& InArgs, const TSharedRef<FMassDebuggerModel>& InDebuggerModel);
 
 	void PopulateProcessorList();
 
diff --git a/Engine/Source/Runtime/MassEntity/Private/MassObserverManager.cpp b/Engine/Source/Runtime/MassEntity/Private/MassObserverManager.cpp
--- a/Engine/Source/Runtime/MassEntity/Private/MassObserverManager.cpp
+++ b/Engine/Source/Runtime/MassEntity/Private/MassObserverManager.cpp
@@ -114,9 +114,9 @@ void FMassObserverManager::Initialize()
 	FMassDebugger::RegisterProcessorDataProvider(TEXT("Observers"), EntityManager.AsShared()
 		, [WeakManager = EntityManager.AsWeak()](TArray<const UMassProcessor*>& OutProcessors)
 	{
-		if (TSharedPtr<FMassEntityManager> EntityManager = WeakManager.Pin())
+		if (TSharedPtr<FMassEntityManager> SharedEntityManager = WeakManager.Pin())
 		{
-			FMassObserverManager& ObserverManager = EntityManager->GetObserverManager();
+			FMassObserverManager& ObserverManager = SharedEntityManager->GetObserverManager();
 			ObserverManager.DebugGatherUniqueProcessors(OutProcessors);
 		}
 	});
