[Mass] CIS issue 640908: fixed shipping cooked editor compilation errors due to Mass debug information not baing available in shipping

#rnx

diff --git a/Engine/Source/Editor/MassEntityDebugger/Private/SMassDebugger.cpp b/Engine/Source/Editor/MassEntityDebugger/Private/SMassDebugger.cpp
--- a/Engine/Source/Editor/MassEntityDebugger/Private/SMassDebugger.cpp
+++ b/Engine/Source/Editor/MassEntityDebugger/Private/SMassDebugger.cpp
@@ -203,6 +203,7 @@ void SMassDebugger::OnEntityManagerInitialized(const FMassEntityManager& EntityM
 	}
 }
 
+#if WITH_MASSENTITY_DEBUG
 void SMassDebugger::OnProcessorProviderRegistered(const FMassDebugger::FEnvironment& Environment)
 {
 	const int32 FoundIndex = EnvironmentsList.FindLastByPredicate([WeakManager = Environment.EntityManager](const TSharedPtr<FMassDebuggerEnvironment>& TestedEnvironment)
@@ -215,6 +216,7 @@ void SMassDebugger::OnProcessorProviderRegistered(const FMassDebugger::FEnvironm
 		EnvironmentsList[FoundIndex]->ProcessorProviders = Environment.ProcessorProviders;
 	}
 }
+#endif // WITH_MASSENTITY_DEBUG
 
 void SMassDebugger::OnEntityManagerDeinitialized(const FMassEntityManager& EntityManager)
 {
diff --git a/Engine/Source/Editor/MassEntityDebugger/Private/SMassDebugger.h b/Engine/Source/Editor/MassEntityDebugger/Private/SMassDebugger.h
--- a/Engine/Source/Editor/MassEntityDebugger/Private/SMassDebugger.h
+++ b/Engine/Source/Editor/MassEntityDebugger/Private/SMassDebugger.h
@@ -64,11 +64,14 @@ protected:
 
 	void BindDelegates();
 	void OnEntityManagerInitialized(const FMassEntityManager& EntityManager);
-	void OnProcessorProviderRegistered(const FMassDebugger::FEnvironment& Environment);
 	void OnEntityManagerDeinitialized(const FMassEntityManager& EntityManager);
 	void HandleEnvironmentChanged(TSharedPtr<FMassDebuggerEnvironment> Item, ESelectInfo::Type SelectInfo);
 	void RebuildEnvironmentsList();
 
+#if WITH_MASSENTITY_DEBUG
+void OnProcessorProviderRegistered(const FMassDebugger::FEnvironment& Environment);
+#endif // WITH_MASSENTITY_DEBUG
+
 	// Holds the list of UI commands.
 	TSharedRef<FUICommandList> CommandList;
 
