[Mass]
Mass Debugger Processors view overhaul

diff --git a/Engine/Source/Editor/MassEntityDebugger/Private/MassDebuggerModel.cpp b/Engine/Source/Editor/MassEntityDebugger/Private/MassDebuggerModel.cpp
--- a/Engine/Source/Editor/MassEntityDebugger/Private/MassDebuggerModel.cpp
+++ b/Engine/Source/Editor/MassEntityDebugger/Private/MassDebuggerModel.cpp
@@ -148,7 +148,8 @@ FMassDebuggerProcessorData::FMassDebuggerProcessorData(const FMassEntityManager&
 {
 	SetProcessor(InProcessor);
 #if WITH_MASSENTITY_DEBUG
-	EntityManager = &InEntityManager;
+	EntityManager = InEntityManager.AsWeak();
+
 	// yeah, it's ugly. But it's debugging code, so... 
 	UMassProcessor& MutableProcessor = const_cast<UMassProcessor&>(InProcessor);
 	TConstArrayView<FMassEntityQuery*> ProcessorQueries = FMassDebugger::GetUpToDateProcessorQueries(InEntityManager, MutableProcessor);
@@ -774,4 +775,15 @@ void FMassDebuggerModel::OnEntitySelected(const FMassEntityManager& EntityManage
 	}
 }
 
+void FMassDebuggerModel::SelectFragment(FName InFragementName)
+{
+	SelectedFragmentName = InFragementName;
+	OnFragmentSelectedDelegate.Broadcast(SelectedFragmentName);
+}
+
+FName FMassDebuggerModel::GetSelectedFragment()
+{
+	return SelectedFragmentName;
+}
+
 #undef LOCTEXT_NAMESPACE
diff --git a/Engine/Source/Editor/MassEntityDebugger/Private/MassDebuggerModel.h b/Engine/Source/Editor/MassEntityDebugger/Private/MassDebuggerModel.h
--- a/Engine/Source/Editor/MassEntityDebugger/Private/MassDebuggerModel.h
+++ b/Engine/Source/Editor/MassEntityDebugger/Private/MassDebuggerModel.h
@@ -2,7 +2,6 @@
 #pragma once
 
 #include "MassArchetypeTypes.h"
-#include "MassEntityHandle.h"
 #include "MassEntityQuery.h"
 #include "MassProcessorDependencySolver.h"
 #include "MassProcessingTypes.h"
@@ -14,6 +13,7 @@
 class UMassProcessor;
 struct FMassArchetypeHandle;
 struct FMassDebuggerModel;
+struct FMassEntityHandle;
 struct FMassEntityManager;
 class UWorld;
 
@@ -89,6 +89,7 @@ struct FMassDebuggerArchetypeData
 	/** Primary debug name, used for grouping derived archetypes. */
 	FString PrimaryDebugName;
 
+	/** Snapshot of all entity handles using this archetype when this struct was created */
 	TArray<FMassEntityHandle> EntityHandles;
 
 	/** True if the archetype is selected. */
@@ -110,7 +111,7 @@ public:
 	FString Label;
 	uint32 ProcessorHash = 0;
 	bool bIsActive = true;
-	const FMassEntityManager* EntityManager = nullptr;
+	TWeakPtr<const FMassEntityManager> EntityManager;
 	TWeakObjectPtr<const UMassProcessor> Processor;
 	EMassDebuggerProcessorSelection Selection = EMassDebuggerProcessorSelection::None;
 	TSharedPtr<FMassDebuggerQueryData> ProcessorRequirements;
@@ -166,6 +167,7 @@ struct FMassDebuggerModel
 	DECLARE_MULTICAST_DELEGATE(FOnRefresh);
 	DECLARE_MULTICAST_DELEGATE_TwoParams(FOnProcessorsSelected, TConstArrayView<TSharedPtr<FMassDebuggerProcessorData>>, ESelectInfo::Type);
 	DECLARE_MULTICAST_DELEGATE_TwoParams(FOnArchetypesSelected, TConstArrayView<TSharedPtr<FMassDebuggerArchetypeData>>, ESelectInfo::Type);
+	DECLARE_MULTICAST_DELEGATE_OneParam(FOnFragmentSelected, FName FragmentName);
 
 	FMassDebuggerModel();
 	~FMassDebuggerModel();
@@ -210,6 +212,7 @@ public:
 	FOnRefresh OnRefreshDelegate;
 	FOnProcessorsSelected OnProcessorsSelectedDelegate;
 	FOnArchetypesSelected OnArchetypesSelectedDelegate;
+	FOnFragmentSelected OnFragmentSelectedDelegate;
 
 	EMassDebuggerSelectionMode SelectionMode = EMassDebuggerSelectionMode::None;
 
@@ -240,8 +243,12 @@ public:
 
 	FDelegateHandle OnEntitySelectedHandle;
 
+	void SelectFragment(FName InFragementName);
+	FName GetSelectedFragment();
+
 protected:	
 	TArray<TSharedPtr<FMassDebuggerProcessorData>> AllCachedProcessors;
+	FName SelectedFragmentName;
 
 public:
 	UE_DEPRECATED(5.6, "CachedProcessors property is now deprecated. Use CachedProcessorCollections instead. ")
diff --git a/Engine/Source/Editor/MassEntityDebugger/Private/MassDebuggerStyle.cpp b/Engine/Source/Editor/MassEntityDebugger/Private/MassDebuggerStyle.cpp
--- a/Engine/Source/Editor/MassEntityDebugger/Private/MassDebuggerStyle.cpp
+++ b/Engine/Source/Editor/MassEntityDebugger/Private/MassDebuggerStyle.cpp
@@ -48,6 +48,14 @@ public:
 		Set("MassDebug.Fragment.ReadOnly", new CORE_BOX_BRUSH("Common/LightGroupBorder", FMargin(4.0f / 16.0f), FLinearColor(.5f, .5f, .5f)));
 		Set("MassDebug.Fragment.ReadWrite", new CORE_BOX_BRUSH("Common/LightGroupBorder", FMargin(4.0f / 16.0f), FLinearColor(0.f, .7f, 0.f)));
 
+		Set("MassDebug.Processor", new FSlateRoundedBoxBrush(FLinearColor::Gray, 10.0f));
+		Set("MassDebug.Processor.AccessRequired", new FSlateRoundedBoxBrush(FLinearColor::Blue, 10.0f));
+		Set("MassDebug.Processor.AccessRead", new FSlateRoundedBoxBrush(FLinearColor::Green, 10.0f));
+		Set("MassDebug.Processor.AccessWrite", new FSlateRoundedBoxBrush(FLinearColor::Red, 10.0f));
+		Set("MassDebug.Processor.AccessBlock", new FSlateRoundedBoxBrush(FLinearColor::Gray * 0.25f, 10.0f));
+
+		Set("MassDebug.Processor.InnerBackground", new FSlateRoundedBoxBrush(FLinearColor::Black.CopyWithNewOpacity(0.75f), 10.0f));
+
 		Set("MassDebug.Label.Background", new FSlateRoundedBoxBrush(FStyleColors::Foreground, 4.0f));
 
 		Set("MassDebug.Label.Text", FTextBlockStyle(NormalText)
diff --git a/Engine/Source/Editor/MassEntityDebugger/Private/SMassArchetype.cpp b/Engine/Source/Editor/MassEntityDebugger/Private/SMassArchetype.cpp
--- a/Engine/Source/Editor/MassEntityDebugger/Private/SMassArchetype.cpp
+++ b/Engine/Source/Editor/MassEntityDebugger/Private/SMassArchetype.cpp
@@ -6,20 +6,30 @@
 #include "SMassBitSet.h"
 #include "Widgets/Text/SRichTextBlock.h"
 #include "Styling/AppStyle.h"
+#include "Widgets/Input/SButton.h"
+#include "Widgets/Views/SListView.h"
 
 #define LOCTEXT_NAMESPACE "SMassDebugger"
 
 //----------------------------------------------------------------------//
 // SMassArchetype
 //----------------------------------------------------------------------//
-void SMassArchetype::Construct(const FArguments& InArgs, TSharedPtr<FMassDebuggerArchetypeData> InArchetypeData, TSharedPtr<FMassDebuggerArchetypeData> InBaseArchetypeData, const EMassBitSetDiffPrune Prune)
+void SMassArchetype::Construct(const FArguments& InArgs, TSharedPtr<FMassDebuggerArchetypeData> InArchetypeData, TSharedPtr<FMassDebuggerArchetypeData> InBaseArchetypeData
+	, const EMassBitSetDiffPrune Prune, TSharedRef<FMassDebuggerModel> InDebuggerModel)
 {
 	if (!InArchetypeData)
 	{
 		return;
 	}
-	
+
 	ArchetypeData = InArchetypeData;
+	DebuggerModel = InDebuggerModel;
+
+	EntityHandles.Reserve(ArchetypeData->EntityHandles.Num());
+	for (FMassEntityHandle& Handle : ArchetypeData->EntityHandles)
+	{
+		EntityHandles.Add(MakeShared<FMassEntityHandle>(Handle));
+	}
 
 	const FMassDebuggerArchetypeData* BaseArchetypeDebugData = InBaseArchetypeData.Get();
 	const FMassDebuggerArchetypeData& ArchetypeDebugData = *InArchetypeData.Get();
@@ -28,33 +38,22 @@ void SMassArchetype::Construct(const FArguments& InArgs, TSharedPtr<FMassDebugge
 	{
 		BaseArchetypeDebugData = nullptr;
 	}
-	
-	TSharedRef<SVerticalBox> Box = SNew(SVerticalBox);
 
+	TSharedRef<SVerticalBox> Box = SNew(SVerticalBox);
 	const TArray<FText> LabelBits = {
 		LOCTEXT("MassArchetypeLabel", "Archetype")
 		, InArchetypeData->LabelLong
 	};
 
 	Box->AddSlot()
-	.AutoHeight()
-	.Padding(0, 4)
-	[
-		SNew(SRichTextBlock)
-		.Text(FText::Join(FText::FromString(TEXT(": ")), LabelBits))
-		.DecoratorStyleSet(&FAppStyle::Get())
-		.TextStyle(FAppStyle::Get(), "LargeText")
-	];
+		.AutoHeight()
+		.Padding(0, 4)
+		[
+			SNew(STextBlock)
+				.Text(InArchetypeData->HashLabel)
+				.Font(FCoreStyle::GetDefaultFontStyle("Mono", 9))
+		];
 
-	Box->AddSlot()
-	.AutoHeight()
-	.Padding(0, 4)
-	[
-		SNew(STextBlock)
-		.Text(InArchetypeData->HashLabel)
-		.Font(FCoreStyle::GetDefaultFontStyle("Mono", 9))
-	];
-	
 	FText ArchetypeDescription = FText::Format(LOCTEXT("ArchetypeDescrption", "EntitiesCount: {0}\
 		\nBytesPerEntity: {1}\
 		\nEntitiesCountPerChunk: {2}\
@@ -68,15 +67,15 @@ void SMassArchetype::Construct(const FArguments& InArgs, TSharedPtr<FMassDebugge
 		, FText::AsMemory(ArchetypeDebugData.ArchetypeStats.AllocatedSize)
 		, FText::AsMemory(ArchetypeDebugData.ArchetypeStats.WastedEntityMemory)
 		, FText::AsNumber(float(ArchetypeDebugData.ArchetypeStats.WastedEntityMemory) * 100.f / ArchetypeDebugData.ArchetypeStats.AllocatedSize)
-		);
-	
+	);
+
 	Box->AddSlot()
-	.AutoHeight()
-	.Padding(0, 4)
-	[
-		SNew(STextBlock)
-		.Text(ArchetypeDescription)
-	];
+		.AutoHeight()
+		.Padding(0, 4)
+		[
+			SNew(STextBlock)
+				.Text(ArchetypeDescription)
+		];
 
 	if (ArchetypeDebugData.ArchetypeStats.EntitiesCount != 0 && ArchetypeDebugData.ArchetypeStats.ChunksCount != 0)
 	{
@@ -101,31 +100,122 @@ void SMassArchetype::Construct(const FArguments& InArgs, TSharedPtr<FMassDebugge
 	{
 		const FMassArchetypeCompositionDescriptor& ParentComposition = BaseArchetypeDebugData->Composition;
 
-		UE::Mass::Debugger::UI::AddBitSetDiff(Box, ParentComposition.Fragments, Composition.Fragments, TEXT("Fragments"), Brush, Prune);
-		UE::Mass::Debugger::UI::AddBitSetDiff(Box, ParentComposition.Tags, Composition.Tags, TEXT("Tags"), Brush, Prune);
-		UE::Mass::Debugger::UI::AddBitSetDiff(Box, ParentComposition.ChunkFragments, Composition.ChunkFragments, TEXT("Chunk Fragments"), Brush, Prune);
-		UE::Mass::Debugger::UI::AddBitSetDiff(Box, ParentComposition.SharedFragments, Composition.SharedFragments, TEXT("Shared Fragments"), Brush, Prune);
-		UE::Mass::Debugger::UI::AddBitSetDiff(Box, ParentComposition.ConstSharedFragments, Composition.ConstSharedFragments, TEXT("Const Shared Fragments"), Brush, Prune);
+		UE::Mass::Debugger::UI::AddBitSetDiff(Box, ParentComposition.Fragments, Composition.Fragments, TEXT("Fragments"), Brush, Prune, InDebuggerModel);
+		UE::Mass::Debugger::UI::AddBitSetDiff(Box, ParentComposition.Tags, Composition.Tags, TEXT("Tags"), Brush, Prune, InDebuggerModel);
+		UE::Mass::Debugger::UI::AddBitSetDiff(Box, ParentComposition.ChunkFragments, Composition.ChunkFragments, TEXT("Chunk Fragments"), Brush, Prune, InDebuggerModel);
+		UE::Mass::Debugger::UI::AddBitSetDiff(Box, ParentComposition.SharedFragments, Composition.SharedFragments, TEXT("Shared Fragments"), Brush, Prune, InDebuggerModel);
+		UE::Mass::Debugger::UI::AddBitSetDiff(Box, ParentComposition.ConstSharedFragments, Composition.ConstSharedFragments, TEXT("Const Shared Fragments"), Brush, Prune, InDebuggerModel);
 	}
 	else
 	{
-		UE::Mass::Debugger::UI::AddBitSet(Box, Composition.Fragments, TEXT("Fragments"), Brush);
-		UE::Mass::Debugger::UI::AddBitSet(Box, Composition.Tags, TEXT("Tags"), Brush);
-		UE::Mass::Debugger::UI::AddBitSet(Box, Composition.ChunkFragments, TEXT("Chunk Fragments"), Brush);
-		UE::Mass::Debugger::UI::AddBitSet(Box, Composition.SharedFragments, TEXT("Shared Fragments"), Brush);
-		UE::Mass::Debugger::UI::AddBitSet(Box, Composition.ConstSharedFragments, TEXT("Const Shared Fragments"), Brush);
+		UE::Mass::Debugger::UI::AddBitSet(Box, Composition.Fragments, TEXT("Fragments"), Brush, InDebuggerModel);
+		UE::Mass::Debugger::UI::AddBitSet(Box, Composition.Tags, TEXT("Tags"), Brush, InDebuggerModel);
+		UE::Mass::Debugger::UI::AddBitSet(Box, Composition.ChunkFragments, TEXT("Chunk Fragments"), Brush, InDebuggerModel);
+		UE::Mass::Debugger::UI::AddBitSet(Box, Composition.SharedFragments, TEXT("Shared Fragments"), Brush, InDebuggerModel);
+		UE::Mass::Debugger::UI::AddBitSet(Box, Composition.ConstSharedFragments, TEXT("Const Shared Fragments"), Brush, InDebuggerModel);
 	}
 
-	ChildSlot
-	[
-		SNew(SBorder)
-		.Padding(5.0f)
+	TSharedRef<SVerticalBox> MainBox = SNew(SVerticalBox);
+
+	MainBox->AddSlot()
+		.AutoHeight()
+		.Padding(0, 4)
+		[
+			SNew(SButton)
+				.OnClicked_Lambda([this, Box]() -> FReply
+				{
+					bBitSetsVisible = !bBitSetsVisible;
+					Box->SetVisibility(bBitSetsVisible ? EVisibility::Visible : EVisibility::Collapsed);
+					return FReply::Handled();
+				})
+				[
+					SNew(SRichTextBlock)
+					.Text(FText::Join(FText::FromString(TEXT(": ")), LabelBits))
+					.DecoratorStyleSet(&FAppStyle::Get())
+					.TextStyle(FAppStyle::Get(), "LargeText")
+				]
+		];
+
+	MainBox->AddSlot()
+		.AutoHeight()
+		.Padding(0, 4)
 		[
 			Box
-		]
-	];
-	
+		];
+
+	TSharedRef<SListView<TSharedPtr<FMassEntityHandle>>> ListViewWidget = SNew(SListView<TSharedPtr<FMassEntityHandle>>)
+		.ListItemsSource(&EntityHandles)
+		.OnGenerateRow(this, &SMassArchetype::OnGenerateRow)
+		.OnSelectionChanged(this, &SMassArchetype::OnEntitySelected);
+
+	MainBox->AddSlot()
+		.AutoHeight()
+		.Padding(0, 4)
+		[
+			SNew(SButton)
+				.Text(LOCTEXT("EntityList", "Entity List"))
+				.OnClicked_Lambda([this, ListViewWidget]() -> FReply
+				{
+					bEntitiesVisible = !bEntitiesVisible;
+					ListViewWidget->SetVisibility(bEntitiesVisible ? EVisibility::Visible : EVisibility::Collapsed);
+					return FReply::Handled();
+				})
+		];
+
+	MainBox->AddSlot()
+		.FillHeight(1.0f)
+		.Padding(0, 4)
+		[
+			ListViewWidget
+		];
+
+	ChildSlot
+		[
+			SNew(SBorder)
+				.Padding(5.0f)
+				[
+					MainBox
+				]
+		];
 }
 
-#undef LOCTEXT_NAMESPACE
+void SMassArchetype::OnEntitySelected(TSharedPtr<FMassEntityHandle> SelectedItem, ESelectInfo::Type SelectInfo)
+{
+	if (SelectedItem.IsValid() && DebuggerModel->Environment.IsValid())
+	{
+		FMassDebugger::SelectEntity(*(DebuggerModel->Environment->GetEntityManager()), *SelectedItem);
+	}
+}
+
+TSharedRef<ITableRow> SMassArchetype::OnGenerateRow(TSharedPtr<FMassEntityHandle> InItem, const TSharedRef<STableViewBase>& OwnerTable)
+{
+	return SNew(STableRow<TSharedPtr<FMassEntityHandle>>, OwnerTable)
+	.Content()
+	[
+		SNew(SButton)
+		.Text(FText::FromString(InItem->DebugGetDescription()))
+		.OnClicked_Lambda([this, InItem]()
+		{
+			if (InItem.IsValid() && DebuggerModel->Environment.IsValid())
+			{
+				FMassDebugger::SelectEntity(*(DebuggerModel->Environment->GetEntityManager()), *InItem);
+			}
+			return FReply::Handled();
+		})
+		.OnHovered_Lambda([this, InItem]()
+		{
+			if (InItem.IsValid() && DebuggerModel->Environment.IsValid())
+			{
+				FMassDebugger::HighlightEntity(*(DebuggerModel->Environment->GetEntityManager()), *InItem);
+			}
+		})
+		.OnUnhovered_Lambda([this]()
+		{
+			if (DebuggerModel->Environment.IsValid())
+			{
+				FMassDebugger::HighlightEntity(*(DebuggerModel->Environment->GetEntityManager()), FMassEntityHandle());
+			}
+		})
+	];
+}
 
diff --git a/Engine/Source/Editor/MassEntityDebugger/Private/SMassArchetype.h b/Engine/Source/Editor/MassEntityDebugger/Private/SMassArchetype.h
--- a/Engine/Source/Editor/MassEntityDebugger/Private/SMassArchetype.h
+++ b/Engine/Source/Editor/MassEntityDebugger/Private/SMassArchetype.h
@@ -5,6 +5,8 @@
 #include "SMassBitSet.h"
 
 struct FMassDebuggerArchetypeData;
+struct FMassEntityHandle;
+struct FMassDebuggerModel;
 
 class SMassArchetype : public SCompoundWidget
 {
@@ -13,8 +15,14 @@ public:
 	SLATE_BEGIN_ARGS(SMassArchetype){}
 	SLATE_END_ARGS()
 
-	void Construct(const FArguments& InArgs, TSharedPtr<FMassDebuggerArchetypeData> InArchetypeData, TSharedPtr<FMassDebuggerArchetypeData> InBaseArchetypeData, const EMassBitSetDiffPrune Prune);
+	void Construct(const FArguments& InArgs, TSharedPtr<FMassDebuggerArchetypeData> InArchetypeData, TSharedPtr<FMassDebuggerArchetypeData> InBaseArchetypeData, const EMassBitSetDiffPrune Prune, TSharedRef<FMassDebuggerModel> InDebuggerModel);
 
 protected:
 	TSharedPtr<FMassDebuggerArchetypeData> ArchetypeData;
+	void OnEntitySelected(TSharedPtr<FMassEntityHandle> SelectedItem, ESelectInfo::Type SelectInfo);
+	TSharedRef<ITableRow> OnGenerateRow(TSharedPtr<FMassEntityHandle> InItem, const TSharedRef<STableViewBase>& OwnerTable);
+	TSharedPtr<FMassDebuggerModel> DebuggerModel;
+	TArray<TSharedPtr<FMassEntityHandle>> EntityHandles;
+	bool bBitSetsVisible = true;
+	bool bEntitiesVisible = true;
 };
diff --git a/Engine/Source/Editor/MassEntityDebugger/Private/SMassArchetypesView.cpp b/Engine/Source/Editor/MassEntityDebugger/Private/SMassArchetypesView.cpp
--- a/Engine/Source/Editor/MassEntityDebugger/Private/SMassArchetypesView.cpp
+++ b/Engine/Source/Editor/MassEntityDebugger/Private/SMassArchetypesView.cpp
@@ -308,7 +308,7 @@ void SMassArchetypesView::RebuildSelectedView()
 	{
 		SelectedArchetypesView->AddSlot()
 			[
-				SNew(SMassArchetype, ArchetypeData, DiffBase, Prune)
+				SNew(SMassArchetype, ArchetypeData, DiffBase, Prune, DebuggerModel.ToSharedRef())
 			];
 	}
 }
diff --git a/Engine/Source/Editor/MassEntityDebugger/Private/SMassBitSet.h b/Engine/Source/Editor/MassEntityDebugger/Private/SMassBitSet.h
--- a/Engine/Source/Editor/MassEntityDebugger/Private/SMassBitSet.h
+++ b/Engine/Source/Editor/MassEntityDebugger/Private/SMassBitSet.h
@@ -4,12 +4,15 @@
 #include "Widgets/SCompoundWidget.h"
 #include "Widgets/Text/STextBlock.h"
 #include "MassDebuggerStyle.h"
+#include "Widgets/Input/SButton.h"
 #include "Widgets/Layout/SWrapBox.h"
 #include "Widgets/Layout/SBorder.h"
 #include "Widgets/Layout/SBox.h"
 #include "Widgets/Images/SImage.h"
 #include "Widgets/SBoxPanel.h"
 #include "MassProcessingTypes.h"
+#include "MassDebugger.h"
+#include "MassDebuggerModel.h"
 
 #define LOCTEXT_NAMESPACE "SMassDebugger"
 
@@ -36,12 +39,12 @@ public:
 		SLATE_ATTRIBUTE(const FSlateBrush*, BackgroundBrush)
 	SLATE_END_ARGS()
 
-	void Construct(const SMassBitSet::FArguments& InArgs, const FString& Label, const TBitSet& BitSet)
+	void Construct(const SMassBitSet::FArguments& InArgs, TSharedRef<FMassDebuggerModel> DebuggerModel, const FString& Label, const TBitSet& BitSet)
 	{
-		Construct(InArgs, Label, MakeArrayView(&BitSet, 1));
+		Construct(InArgs, DebuggerModel, Label, MakeArrayView(&BitSet, 1));
 	}
 
-	void Construct(const SMassBitSet::FArguments& InArgs, const FString& Label, TConstArrayView<TBitSet> BitSets, TConstArrayView<const FSlateBrush*> InBrushes = TConstArrayView<const FSlateBrush*>())
+	void Construct(const SMassBitSet::FArguments& InArgs, TSharedRef<FMassDebuggerModel> DebuggerModel, const FString& Label, TConstArrayView<TBitSet> BitSets, TConstArrayView<const FSlateBrush*> InBrushes = TConstArrayView<const FSlateBrush*>())
 	{		
 		/*Content = TEXT("<LargeText>Large test</>, <RichTextBlock.Bold>Bold</>");
 		Content += BitSet.DebugGetStringDesc();*/
@@ -50,7 +53,7 @@ public:
 		for (int i = 0; i < BitSets.Num(); ++i)
 		{
 			const FSlateBrush* Brush = InBrushes.IsValidIndex(i) ? InBrushes[i] : InArgs._BackgroundBrush.Get();
-			AddBitSet(InArgs, ButtonBox, BitSets[i], Brush);
+			AddBitSet(InArgs, ButtonBox, BitSets[i], Brush, DebuggerModel);
 		}
 
 		ChildSlot
@@ -75,7 +78,8 @@ public:
 	}
 
 protected:
-	void AddBitSet(const SMassBitSet::FArguments& InArgs, TSharedRef<SWrapBox>& ButtonBox, const TBitSet& BitSet, const FSlateBrush* Brush)
+	void AddBitSet(const SMassBitSet::FArguments& InArgs, TSharedRef<SWrapBox>& ButtonBox, const TBitSet& BitSet, const FSlateBrush* Brush
+		, TSharedRef<FMassDebuggerModel> DebuggerModel)
 	{
 #if WITH_MASSENTITY_DEBUG
 		TArray<FName> TypeNames;
@@ -90,8 +94,16 @@ protected:
 				.ColorAndOpacity(InArgs._TextColor)
 				.BorderImage(Brush)
 				[
-					SNew(STextBlock)
-					.Text(FText::FromName(Name))
+					SNew(SButton)
+					.OnClicked_Lambda([Name, DebuggerModel]()
+					{
+						DebuggerModel->SelectFragment(Name);
+						return FReply::Handled();
+					})
+					[
+						SNew(STextBlock)
+						.Text(FText::FromName(Name))
+					]
 				]
 			];
 		}
@@ -120,12 +132,13 @@ public:
 		SLATE_ATTRIBUTE(const FSlateBrush*, BackgroundBrush)
 	SLATE_END_ARGS()
 
-	void Construct(const SMassBitSetDiff::FArguments& InArgs, const FString& Label, const TBitSet& BaseBitSet, const TBitSet& BitSet, const EMassBitSetDiffPrune Prune)
+	void Construct(const SMassBitSetDiff::FArguments& InArgs, const FString& Label, const TBitSet& BaseBitSet, const TBitSet& BitSet
+		, const EMassBitSetDiffPrune Prune, TSharedRef<FMassDebuggerModel> DebuggerModel)
 	{
 		TSharedRef<SWrapBox> ButtonBox = SNew(SWrapBox).UseAllottedSize(true);
 		
 		const FSlateBrush* Brush = InArgs._BackgroundBrush.Get();
-		AddBitSet(InArgs, ButtonBox, BaseBitSet, BitSet, Brush, Prune);
+		AddBitSet(InArgs, ButtonBox, BaseBitSet, BitSet, Brush, Prune, DebuggerModel);
 
 		ChildSlot
 		[
@@ -137,8 +150,16 @@ public:
 				.MinDesiredWidth(150.0f)
 				.Padding(FMargin(0.0f, InArgs._SlotPadding.Get().Top, 0.0f, InArgs._SlotPadding.Get().Bottom))
 				[
-					SNew(STextBlock)
-					.Text(FText::FromString(Label))
+					SNew(SButton)
+					.OnClicked_Lambda([DebuggerModel, Label]()
+					{
+						DebuggerModel->SelectFragment(FName(*Label));
+						return FReply::Handled();
+					})
+					[
+						SNew(STextBlock)
+						.Text(FText::FromString(Label))
+					]
 				]
 			]
 			+ SHorizontalBox::Slot()
@@ -149,7 +170,8 @@ public:
 	}
 
 protected:
-	void AddBitSet(const SMassBitSetDiff::FArguments& InArgs, TSharedRef<SWrapBox>& ButtonBox, const TBitSet& BaseBitSet, const TBitSet& BitSet, const FSlateBrush* Brush, const EMassBitSetDiffPrune Prune)
+	void AddBitSet(const SMassBitSetDiff::FArguments& InArgs, TSharedRef<SWrapBox>& ButtonBox, const TBitSet& BaseBitSet, const TBitSet& BitSet
+		, const FSlateBrush* Brush, const EMassBitSetDiffPrune Prune, TSharedRef<FMassDebuggerModel> DebuggerModel)
 	{
 #if WITH_MASSENTITY_DEBUG
 		TBitSet CombinedBitSet = BaseBitSet + BitSet;
@@ -198,34 +220,42 @@ protected:
 			}
 			
 			ButtonBox->AddSlot()
-				.Padding(InArgs._SlotPadding)
+			.Padding(InArgs._SlotPadding)
+			[
+				SNew(SBorder)
+				.BorderImage(Brush)
+				.BorderBackgroundColor(BorderColor)
+				.ColorAndOpacity(TextColor)
+				.ToolTipText(Tooltip)
 				[
-					SNew(SBorder)
-					.BorderImage(Brush)
-					.BorderBackgroundColor(BorderColor)
-					.ColorAndOpacity(TextColor)
-					.ToolTipText(Tooltip)
+					SNew(SHorizontalBox)
+					+ SHorizontalBox::Slot()
+					.AutoWidth()
+					.Padding(3,0)
+					.VAlign(VAlign_Center)
 					[
-						SNew(SHorizontalBox)
-						+ SHorizontalBox::Slot()
-						.AutoWidth()
-						.Padding(3,0)
-						.VAlign(VAlign_Center)
-						[
-							SNew(SImage)
-							.Image(Icon)
-							.Visibility_Lambda([Icon] { return Icon != nullptr ? EVisibility::Visible : EVisibility::Collapsed; })
-							.DesiredSizeOverride(FVector2D(10.0f, 10.0f))
-						]
-						+ SHorizontalBox::Slot()
-						.AutoWidth()
-						.VAlign(VAlign_Center)
+						SNew(SImage)
+						.Image(Icon)
+						.Visibility_Lambda([Icon] { return Icon != nullptr ? EVisibility::Visible : EVisibility::Collapsed; })
+						.DesiredSizeOverride(FVector2D(10.0f, 10.0f))
+					]
+					+ SHorizontalBox::Slot()
+					.AutoWidth()
+					.VAlign(VAlign_Center)
+					[
+						SNew(SButton)
+						.OnClicked_Lambda([Name, DebuggerModel]()
+							{
+								DebuggerModel->SelectFragment(FName(*Name));
+								return FReply::Handled();
+							})
 						[
 							SNew(STextBlock)
-							.Text(FText::FromString(Name))
+								.Text(FText::FromString(Name))
 						]
 					]
-				];
+				]
+			];
 		}
 
 		if (SameCount > 0)
@@ -251,7 +281,7 @@ protected:
 namespace UE::Mass::Debugger::UI
 {
 	template<typename TBitSet>
-	void AddBitSet(TSharedRef<SVerticalBox>& Box, const TBitSet& BitSetAccess, const FString& Label, const FSlateBrush* Brush)
+	void AddBitSet(TSharedRef<SVerticalBox>& Box, const TBitSet& BitSetAccess, const FString& Label, const FSlateBrush* Brush, TSharedRef<FMassDebuggerModel> DebuggerModel)
 	{
 		if (BitSetAccess.IsEmpty() == false)
 		{
@@ -259,7 +289,7 @@ namespace UE::Mass::Debugger::UI
 			.AutoHeight()
 			.Padding(0, 8)
 			[
-				SNew(SMassBitSet<TBitSet>, Label, BitSetAccess)
+				SNew(SMassBitSet<TBitSet>, DebuggerModel, Label, BitSetAccess)
 				.BackgroundBrush(Brush)
 				.SlotPadding(5.0f)
 			];
@@ -267,7 +297,8 @@ namespace UE::Mass::Debugger::UI
 	}
 
 	template<typename TBitSet>
-	void AddBitSetDiff(TSharedRef<SVerticalBox>& Box, const TBitSet& BaseBitSet, const TBitSet& BitSet, const FString& Label, const FSlateBrush* Brush, const EMassBitSetDiffPrune Prune)
+	void AddBitSetDiff(TSharedRef<SVerticalBox>& Box, const TBitSet& BaseBitSet, const TBitSet& BitSet, const FString& Label, const FSlateBrush* Brush
+		, const EMassBitSetDiffPrune Prune, TSharedRef<FMassDebuggerModel> DebuggerModel)
 	{
 		if (BaseBitSet.IsEmpty() == false || BitSet.IsEmpty() == false)
 		{
@@ -275,7 +306,7 @@ namespace UE::Mass::Debugger::UI
 			.AutoHeight()
 			.Padding(0, 8)
 			[
-				SNew(SMassBitSetDiff<TBitSet>, Label, BaseBitSet, BitSet, Prune)
+				SNew(SMassBitSetDiff<TBitSet>, Label, BaseBitSet, BitSet, Prune, DebuggerModel)
 				.BackgroundBrush(Brush)
 				.SlotPadding(5.0f)
 			];
diff --git a/Engine/Source/Editor/MassEntityDebugger/Private/SMassProcessor.cpp b/Engine/Source/Editor/MassEntityDebugger/Private/SMassProcessor.cpp
--- a/Engine/Source/Editor/MassEntityDebugger/Private/SMassProcessor.cpp
+++ b/Engine/Source/Editor/MassEntityDebugger/Private/SMassProcessor.cpp
@@ -15,7 +15,7 @@
 //----------------------------------------------------------------------//
 // SMassProcessor
 //----------------------------------------------------------------------//
-void SMassProcessor::Construct(const FArguments& InArgs, TSharedPtr<FMassDebuggerProcessorData> InProcessorData)
+void SMassProcessor::Construct(const FArguments& InArgs, TSharedPtr<FMassDebuggerProcessorData> InProcessorData, TSharedRef<FMassDebuggerModel> InDebuggerModel)
 {
 	ProcessorData = InProcessorData;
 	if (!ProcessorData)
@@ -24,18 +24,6 @@ void SMassProcessor::Construct(const FArguments& InArgs, TSharedPtr<FMassDebugge
 	}
 
 	TSharedRef<SVerticalBox> Box = SNew(SVerticalBox);
-	Box->AddSlot()
-	.AutoHeight()
-	[
-		SNew(SBorder)
-		.Padding(FMargin(10, 5))
-		[
-			SNew(SRichTextBlock)
-			.Text(FText::FromString(ProcessorData->Label))
-			.DecoratorStyleSet(&FAppStyle::Get())
-			.TextStyle(FAppStyle::Get(), "LargeText")
-		]
-	];
 
 #if WITH_MASSENTITY_DEBUG
 	if (!ProcessorData->Description.IsEmpty())
@@ -71,7 +59,7 @@ void SMassProcessor::Construct(const FArguments& InArgs, TSharedPtr<FMassDebugge
 		Box->AddSlot()
 		.AutoHeight()
 		[
-			SNew(SMassQuery, ProcessorData->ProcessorRequirements)
+			SNew(SMassQuery, ProcessorData->ProcessorRequirements, InDebuggerModel)
 		];
 	}
 
@@ -80,7 +68,7 @@ void SMassProcessor::Construct(const FArguments& InArgs, TSharedPtr<FMassDebugge
 		Box->AddSlot()
 		.AutoHeight()
 		[
-			SNew(SMassQuery, QueryData)
+			SNew(SMassQuery, QueryData, InDebuggerModel)
 		];
 	}
 
diff --git a/Engine/Source/Editor/MassEntityDebugger/Private/SMassProcessor.h b/Engine/Source/Editor/MassEntityDebugger/Private/SMassProcessor.h
--- a/Engine/Source/Editor/MassEntityDebugger/Private/SMassProcessor.h
+++ b/Engine/Source/Editor/MassEntityDebugger/Private/SMassProcessor.h
@@ -3,7 +3,7 @@
 
 #include "Widgets/SCompoundWidget.h"
 
-
+struct FMassDebuggerModel;
 struct FMassDebuggerProcessorData;
 
 class SMassProcessor : public SCompoundWidget
@@ -13,7 +13,7 @@ public:
 	SLATE_BEGIN_ARGS(SMassProcessor){}
 	SLATE_END_ARGS()
 
-	void Construct(const FArguments& InArgs, TSharedPtr<FMassDebuggerProcessorData> InProcessorData);
+	void Construct(const FArguments& InArgs, TSharedPtr<FMassDebuggerProcessorData> InProcessorData, TSharedRef<FMassDebuggerModel> InDebuggerModel);
 
 protected:
 	TSharedPtr<FMassDebuggerProcessorData> ProcessorData;
diff --git a/Engine/Source/Editor/MassEntityDebugger/Private/SMassProcessorsView.cpp b/Engine/Source/Editor/MassEntityDebugger/Private/SMassProcessorsView.cpp
--- a/Engine/Source/Editor/MassEntityDebugger/Private/SMassProcessorsView.cpp
+++ b/Engine/Source/Editor/MassEntityDebugger/Private/SMassProcessorsView.cpp
@@ -2,6 +2,7 @@
 
 #include "SMassProcessorsView.h"
 #include "SMassProcessor.h"
+#include "SMassProcessorsWidget.h"
 #include "Widgets/Layout/SScrollBox.h"
 #include "Widgets/Views/SListView.h"
 #include "Widgets/SBoxPanel.h"
@@ -14,6 +15,7 @@ namespace UE::Mass::Debug::UI::Private
 	const FText MissingDebugData = FText::FromString(TEXT("Missing debug data"));
 }
 
+
 //----------------------------------------------------------------------//
 // SMassProcessorListTableRow
 //----------------------------------------------------------------------//
@@ -22,7 +24,8 @@ class SMassProcessorListTableRow : public STableRow<TSharedPtr<FMassDebuggerProc
 public:
 	using Super = STableRow<TSharedPtr<FMassDebuggerProcessorData>>;
 
-	void Construct(const FArguments& InArgs, const TSharedRef<STableViewBase>& InOwnerTableView, const TSharedPtr<FMassDebuggerProcessorData>& InEntryItem)
+	void Construct(const FArguments& InArgs, const TSharedRef<STableViewBase>& InOwnerTableView, const TSharedPtr<FMassDebuggerProcessorData>& InEntryItem
+		, TSharedRef<FMassDebuggerModel> DebuggerModel)
 	{
 		if (!InEntryItem.IsValid())
 		{
@@ -35,9 +38,11 @@ public:
 		
 		ChildSlot
 		[
-			SNew(STextBlock)
-			.Text(FText::FromString(Item->Label))
-			.TextStyle(FAppStyle::Get(), Item->bIsActive ? "NormalText" : "NormalText.Subdued")
+			SNew(SBox)
+			.Padding(2.0f)
+			[
+				SNew(SMassProcessorWidget, Item.ToSharedRef(), DebuggerModel)
+			]
 		];
 	}
 
@@ -86,7 +91,7 @@ public:
 	using FItemType = TSharedPtr<FMassDebuggerProcessorData>;
 
 	void Construct(const FArguments&, const TSharedRef<STableViewBase>& InOwnerTableView, const TSharedPtr<FMassDebuggerModel::FProcessorCollection>& InEntryItem
-		, const TSharedRef<SMassProcessorsView>& MainView)
+		, const TSharedRef<SMassProcessorsView>& MainView, TSharedRef<FMassDebuggerModel> DebuggerModel)
 	{
 		CollectionItem = InEntryItem;
 
@@ -99,6 +104,8 @@ public:
 			SNew(SVerticalBox)
 
 			+ SVerticalBox::Slot()
+			.AutoHeight()
+			.VAlign(VAlign_Top)
 			[
 				SNew(STextBlock)
 				.Text(FText::FromName(CollectionItem->Label))
@@ -106,15 +113,16 @@ public:
 			]
 			+ SVerticalBox::Slot()
 			.AutoHeight()
+			.VAlign(VAlign_Top)
 			[
 				SAssignNew(ProcessorsListWidget, SMassProcessorCollectionListView, MainView)
 				.ListItemsSource(&CollectionItem->Container)
-				.SelectionMode(ESelectionMode::Multi)
+				.SelectionMode(ESelectionMode::None)
 				.OnSelectionChanged(MainView->OnProcessorSelectionChanged)
-				.OnGenerateRow_Lambda([](FItemType Item, const TSharedRef<STableViewBase>& OwnerTable)
-					{
-						return SNew(SMassProcessorListTableRow, OwnerTable, Item);
-					})
+				.OnGenerateRow_Lambda([DebuggerModel](FItemType Item, const TSharedRef<STableViewBase>& OwnerTable)
+				{
+					return SNew(SMassProcessorListTableRow, OwnerTable, Item, DebuggerModel);
+				})
 			]
 		];
 	}
@@ -130,52 +138,33 @@ void SMassProcessorsView::Construct(const FArguments& InArgs, const TSharedRef<F
 {
 	Initialize(InDebuggerModel);
 
-	ProcessorsBox = SNew(SVerticalBox);
-	for (TSharedPtr<FMassDebuggerProcessorData>& ProcessorData : DebuggerModel->SelectedProcessors)
-	{
-		ProcessorsBox->AddSlot()
-		[
-			SNew(SMassProcessor, ProcessorData)
-		];
-	}
 
-	OnProcessorSelectionChanged = FOnProcessorSelectionChanged::CreateSP(this, &SMassProcessorsView::ProcessorListSelectionChanged);
-	
 	ChildSlot
-		[
-			SNew(SSplitter)
-			.Orientation(Orient_Horizontal)
+	[
+		SNew(SSplitter)
+		.Orientation(Orient_Horizontal)
 			 
-			+ SSplitter::Slot()
-			.Value(.35f)
-			.MinSize(260.0f)
-			[
-				SNew(SScrollBox)
-				.Orientation(Orient_Vertical)
-				+ SScrollBox::Slot()
-				[
-					SAssignNew(ProcessorCollectionsListWidget, SListView<TSharedPtr<FMassDebuggerModel::FProcessorCollection>>)
-					.ListItemsSource(&DebuggerModel->CachedProcessorCollections)
-					.SelectionMode(ESelectionMode::None)
-					.OnGenerateRow_Lambda([SharedThis = StaticCastSharedRef<SMassProcessorsView>(AsShared())]
-						(TSharedPtr<FMassDebuggerModel::FProcessorCollection> Item, const TSharedRef<STableViewBase>& OwnerTable)
-						{
-							return SNew(SMassProcessorCollectionTableRow, OwnerTable, Item, SharedThis);
-						})
-				]
-			]
-			
-			+ SSplitter::Slot()
-			.Value(.7f)
+		+ SSplitter::Slot()
+		.Value(.35f)
+		.MinSize(260.0f)
+		[
+			SNew(SScrollBox)
+			.Orientation(Orient_Vertical)
+			+ SScrollBox::Slot()
+			.VAlign(VAlign_Top)
 			[
-				SNew(SScrollBox)
-				.Orientation(Orient_Vertical)
-				+SScrollBox::Slot()
-				[
-					ProcessorsBox.ToSharedRef()
-				]
+				SAssignNew(ProcessorCollectionsListWidget, SListView<TSharedPtr<FMassDebuggerModel::FProcessorCollection>>)
+				.ListItemsSource(&DebuggerModel->CachedProcessorCollections)
+				.SelectionMode(ESelectionMode::None)
+				.Orientation(Orient_Horizontal)
+				.OnGenerateRow_Lambda([SharedThis = StaticCastSharedRef<SMassProcessorsView>(AsShared()), InDebuggerModel]
+					(TSharedPtr<FMassDebuggerModel::FProcessorCollection> Item, const TSharedRef<STableViewBase>& OwnerTable)
+					{
+						return SNew(SMassProcessorCollectionTableRow, OwnerTable, Item, SharedThis, InDebuggerModel);
+					})
 			]
-		];
+		]
+	];
 
 	PopulateProcessorList();
 }
@@ -235,7 +224,6 @@ void SMassProcessorsView::OnProcessorsSelected(TConstArrayView<TSharedPtr<FMassD
 			Widget->ClearSelection();
 		}
 	}
-	ProcessorsBox->ClearChildren();
 
 	if (SelectedProcessors.Num())
 	{		
@@ -246,15 +234,6 @@ void SMassProcessorsView::OnProcessorsSelected(TConstArrayView<TSharedPtr<FMassD
 				Widget->SetItemSelection(SelectedProcessors, /*bSelected=*/true, ESelectInfo::OnMouseClick);
 			}
 		}
-
-		for (TSharedPtr<FMassDebuggerProcessorData>& ProcessorData : DebuggerModel->SelectedProcessors)
-		{
-			ProcessorsBox->AddSlot()
-				.AutoHeight()
-				[
-					SNew(SMassProcessor, ProcessorData)
-				];
-		}
 	}
 
 	ProcessorCollectionsListWidget->RequestListRefresh();
@@ -276,19 +255,6 @@ void SMassProcessorsView::PopulateProcessorList()
 	DebuggerModel->ClearProcessorSelection();
 
 	ProcessorCollectionsListWidget->RequestListRefresh();
-
-	if (ProcessorsBox)
-	{
-		ProcessorsBox->ClearChildren();
-		for (TSharedPtr<FMassDebuggerProcessorData>& ProcessorData : DebuggerModel->SelectedProcessors)
-		{
-			ProcessorsBox->AddSlot()
-				.AutoHeight()
-				[
-					SNew(SMassProcessor, ProcessorData)
-				];
-		}
-	}
 }
 
 #undef LOCTEXT_NAMESPACE
diff --git a/Engine/Source/Editor/MassEntityDebugger/Private/SMassProcessorsWidget.cpp b/Engine/Source/Editor/MassEntityDebugger/Private/SMassProcessorsWidget.cpp
new file mode 100644
--- /dev/null
+++ b/Engine/Source/Editor/MassEntityDebugger/Private/SMassProcessorsWidget.cpp
@@ -0,0 +1,295 @@
+// Copyright Epic Games, Inc. All Rights Reserved.
+
+#include "SMassProcessorsWidget.h"
+#include "SMassProcessor.h"
+#include "Styling/AppStyle.h"
+#include "Widgets/SBoxPanel.h"
+#include "Widgets/Images/SImage.h"
+#include "Widgets/Input/SButton.h"
+#include "Widgets/Layout/SBorder.h"
+#include "Widgets/Layout/SBox.h"
+#include "Widgets/Text/STextBlock.h"
+#include "SourceCodeNavigation.h"
+#include "MassDebuggerStyle.h"
+
+#define LOCTEXT_NAMESPACE "SMassDebugger"
+
+SMassProcessorWidget::~SMassProcessorWidget()
+{
+	if (DebuggerModel.IsValid())
+	{
+		DebuggerModel->OnFragmentSelectedDelegate.Remove(OnFragmentSelectChangeHandle);
+	}
+}
+
+void SMassProcessorWidget::Construct(const FArguments& InArgs, TSharedPtr<FMassDebuggerProcessorData> InDebuggerProcessorData, TSharedRef<FMassDebuggerModel> InDebuggerModel)
+{
+	ProcessorData = InDebuggerProcessorData;
+	DebuggerModel = InDebuggerModel;
+
+	OnFragmentSelectChangeHandle = DebuggerModel->OnFragmentSelectedDelegate.AddSP(this, &SMassProcessorWidget::HandleFragmentSelected);
+
+	bIsExpandedText = false;
+	bIsExpandedGraph = false;
+
+	if (!ProcessorData.IsValid())
+	{
+		ChildSlot
+		[
+			SNew(STextBlock)
+			.TextStyle(FAppStyle::Get(), "LargeText")
+			.Text(LOCTEXT("InvalidDebugProcessorData", "Invalid Debug Processor Data"))
+		];
+		return;
+	}
+
+	FMassDebuggerProcessorData& ProcData = *ProcessorData.Get();
+
+	ChildSlot
+	[
+		SAssignNew(Border, SBorder)
+		.BorderImage(GetBorderByFragmentSelection())
+		.BorderBackgroundColor(FLinearColor::Gray)
+		.Padding(3.0f)
+		[
+			SNew(SBorder)
+			.BorderImage(FMassDebuggerStyle::GetBrush("MassDebug.Processor.InnerBackground"))
+			.Padding(5.0f)
+			[
+				SNew(SVerticalBox)
+
+				+ SVerticalBox::Slot()
+				.AutoHeight()
+				[
+					SNew(SHorizontalBox)
+
+					+ SHorizontalBox::Slot()
+					.FillWidth(1.0f)
+					[
+						SNew(SButton)
+						.ButtonStyle(FAppStyle::Get(), "NoBorder")
+						.ToolTipText(LOCTEXT("ViewEntities", "View Entities"))
+						.Content()
+						[
+							SNew(STextBlock)
+							.TextStyle(FAppStyle::Get(), "LargeText")
+							.Text(FText::FromString(ProcData.Label))
+						]
+						.OnClicked(this, &SMassProcessorWidget::HandleSelectProcessorClicked)
+					]
+				]
+
+				+ SVerticalBox::Slot()
+				.AutoHeight()
+				.HAlign(HAlign_Right)
+				[
+					SNew(SHorizontalBox)
+
+					+ SHorizontalBox::Slot()
+					.AutoWidth()
+					[
+						SNew(SButton)
+						.ButtonStyle(FAppStyle::Get(), "NoBorder")
+						.ToolTipText(LOCTEXT("ShowFragmentAccess", "Show Fragment Access"))
+						.Content()
+						[
+							SNew(SImage)
+							.Image(FAppStyle::Get().GetBrush("Icons.Layout"))
+						]
+						.OnClicked(this, &SMassProcessorWidget::HandleExpandTextClicked)
+					]
+
+					+ SHorizontalBox::Slot()
+					.AutoWidth()
+					[
+						SNew(SButton)
+						.ButtonStyle(FAppStyle::Get(), "NoBorder")
+						.ToolTipText(LOCTEXT("ShowStats", "Show Stats"))
+						.Content()
+						[
+							SNew(SImage)
+							.Image(FAppStyle::Get().GetBrush("Icons.LOD"))
+						]
+						.OnClicked(this, &SMassProcessorWidget::HandleExpandGraphClicked)
+					]
+
+					+ SHorizontalBox::Slot()
+					.AutoWidth()
+					[
+						SNew(SButton)
+						.ButtonStyle(FAppStyle::Get(), "NoBorder")
+						.ToolTipText(LOCTEXT("OpenSourceLocation", "Open Source Location"))
+						.Content()
+						[
+							SNew(SImage)
+							.Image(FAppStyle::Get().GetBrush("Icons.C++"))
+						]
+						.OnClicked(this, &SMassProcessorWidget::HandleOpenSourceLocationClicked)
+					]
+				]
+
+				+ SVerticalBox::Slot()
+				.AutoHeight()
+				.Padding(5.0f)
+				[
+					SAssignNew(TextBoxContainer, SBox)
+					.Visibility(bIsExpandedText ? EVisibility::Visible : EVisibility::Collapsed)
+					[
+						SNew(STextBlock)
+						.Text(LOCTEXT("DetailedProcessorInformation", "Detailed information about the processor..."))
+					]
+				]
+
+				+ SVerticalBox::Slot()
+				.AutoHeight()
+				.Padding(5.0f)
+				[
+					SAssignNew(GraphBoxContainer, SBox)
+					.Visibility(bIsExpandedGraph ? EVisibility::Visible : EVisibility::Collapsed)
+					[
+						// This is a placeholder for the stats/perf graph widget
+						SNew(SImage)
+						.Image(FAppStyle::Get().GetBrush("Icons.LOD"))
+					]
+				]
+			]
+		]
+	];
+}
+
+FReply SMassProcessorWidget::HandleSelectProcessorClicked()
+{
+	return FReply::Handled();
+}
+
+FReply SMassProcessorWidget::HandleExpandTextClicked()
+{
+	bIsExpandedText = !bIsExpandedText;
+	if (TextBoxContainer.IsValid() && DebuggerModel.IsValid())
+	{
+		TextBoxContainer->SetContent(SNew(SMassProcessor, ProcessorData, DebuggerModel.ToSharedRef()));
+		TextBoxContainer->SetVisibility(bIsExpandedText ? EVisibility::Visible : EVisibility::Collapsed);
+	}
+	return FReply::Handled();
+}
+
+FReply SMassProcessorWidget::HandleExpandGraphClicked()
+{
+	bIsExpandedGraph = !bIsExpandedGraph;
+	if (GraphBoxContainer.IsValid())
+	{
+		GraphBoxContainer->SetVisibility(bIsExpandedGraph ? EVisibility::Visible : EVisibility::Collapsed);
+	}
+	return FReply::Handled();
+}
+
+FReply SMassProcessorWidget::HandleOpenSourceLocationClicked()
+{
+	if (!ProcessorData.IsValid())
+	{
+		return FReply::Handled();
+	}
+
+	const UMassProcessor* Proc = ProcessorData.Get()->Processor.Get();
+	if (Proc)
+	{
+		FSourceCodeNavigation::NavigateToClass(Proc->GetClass());
+	}
+	
+	return FReply::Handled();
+}
+
+void SMassProcessorWidget::HandleFragmentSelected(FName SelectedFragment)
+{
+	Border->SetBorderImage(GetBorderByFragmentSelection());
+}
+
+const FSlateBrush* SMassProcessorWidget::GetBorderByFragmentSelection()
+{
+	enum AccessLevel
+	{
+		None = 0,
+		Require,
+		Read,
+		Write,
+		Block,
+		MAX
+	};
+
+	const FSlateBrush* AccessColor[] = {
+		FMassDebuggerStyle::GetBrush("MassDebug.Processor"),
+		FMassDebuggerStyle::GetBrush("MassDebug.Processor.AccessRequired"),
+		FMassDebuggerStyle::GetBrush("MassDebug.Processor.AccessRead"),
+		FMassDebuggerStyle::GetBrush("MassDebug.Processor.AccessWrite"),
+		FMassDebuggerStyle::GetBrush("MassDebug.Processor.AccessBlock")
+	};
+
+	static_assert(UE_ARRAY_COUNT(AccessColor) == int32(AccessLevel::MAX), "Ensure AccessLevel enum is 1:1 with the brush array");
+
+	FName SelectedFragment;
+	
+	if (DebuggerModel.IsValid())
+	{
+		SelectedFragment = DebuggerModel->GetSelectedFragment();
+	}
+
+	auto HasSelectedBit = [&SelectedFragment](const auto& Bitset) -> bool {
+		for (int32 Index = 0; Index < Bitset.StructTypesBitArray.Num(); ++Index)
+		{
+			if (Bitset.StructTypesBitArray[Index])
+			{
+				FName setName = Bitset.GetImplementation().GetStructTracker().DebugGetStructTypeName(Index);
+				if (setName == SelectedFragment)
+				{
+					return true;
+				}
+			}
+		}
+		return false;
+		};
+
+	auto GetHighestAccessLevel = [&SelectedFragment, &HasSelectedBit](const auto& ExecutionAccess, AccessLevel Highest) -> AccessLevel {
+		if (Highest < AccessLevel::Write && HasSelectedBit(ExecutionAccess.Write))
+		{
+			return AccessLevel::Write;
+		}
+		if (Highest < AccessLevel::Read && HasSelectedBit(ExecutionAccess.Read))
+		{
+			return AccessLevel::Read;
+		}
+		return Highest;
+		};
+
+	AccessLevel HighestAccess = AccessLevel::None;
+
+	if (!SelectedFragment.IsNone())
+	{
+		for (TSharedPtr<FMassDebuggerQueryData>& Query : ProcessorData->Queries)
+		{
+			FMassExecutionRequirements& ExecutionReqs = Query->ExecutionRequirements;
+			if (HighestAccess < AccessLevel::Block && HasSelectedBit(ExecutionReqs.RequiredNoneTags))
+			{
+				HighestAccess = AccessLevel::Block;
+			}
+			if (HighestAccess < AccessLevel::Require && (HasSelectedBit(ExecutionReqs.RequiredAnyTags) || HasSelectedBit(ExecutionReqs.RequiredAllTags)))
+			{
+				HighestAccess = AccessLevel::Require;
+			}
+			else
+			{
+				HighestAccess = GetHighestAccessLevel(ExecutionReqs.Fragments, HighestAccess);
+				HighestAccess = GetHighestAccessLevel(ExecutionReqs.ChunkFragments, HighestAccess);
+				HighestAccess = GetHighestAccessLevel(ExecutionReqs.SharedFragments, HighestAccess);
+				HighestAccess = GetHighestAccessLevel(ExecutionReqs.RequiredSubsystems, HighestAccess);
+
+				if (HighestAccess < AccessLevel::Read && HasSelectedBit(ExecutionReqs.ConstSharedFragments.Read))
+				{
+					HighestAccess = AccessLevel::Read;
+				}
+			}
+		}
+	}
+
+	return AccessColor[HighestAccess];
+}
+
diff --git a/Engine/Source/Editor/MassEntityDebugger/Private/SMassProcessorsWidget.h b/Engine/Source/Editor/MassEntityDebugger/Private/SMassProcessorsWidget.h
new file mode 100644
--- /dev/null
+++ b/Engine/Source/Editor/MassEntityDebugger/Private/SMassProcessorsWidget.h
@@ -0,0 +1,37 @@
+// Copyright Epic Games, Inc. All Rights Reserved.
+#pragma once
+
+#include "SMassProcessorsView.h"
+#include "SMassProcessor.h"
+
+class SBox;
+class SBorder;
+
+struct FMassDebuggerProcessorData;
+struct FMassDebuggerModel;
+
+class SMassProcessorWidget : public SCompoundWidget
+{
+public:
+	SLATE_BEGIN_ARGS(SMassProcessorWidget) {}
+	SLATE_END_ARGS()
+	virtual ~SMassProcessorWidget();
+
+	void Construct(const FArguments& InArgs, TSharedPtr<FMassDebuggerProcessorData> InDebuggerProcessorData, TSharedRef<FMassDebuggerModel> InDebuggerModel);
+
+private:
+	FReply HandleExpandTextClicked();
+	FReply HandleExpandGraphClicked();
+	FReply HandleSelectProcessorClicked();
+	FReply HandleOpenSourceLocationClicked();
+	void HandleFragmentSelected(FName SelectedFragment);
+	bool bIsExpandedText;
+	bool bIsExpandedGraph;
+	TSharedPtr<SBox> TextBoxContainer;
+	TSharedPtr<SBox> GraphBoxContainer;
+	TSharedPtr<SBorder> Border;
+	TSharedPtr<FMassDebuggerProcessorData> ProcessorData;
+	TSharedPtr<FMassDebuggerModel> DebuggerModel;
+	const FSlateBrush* GetBorderByFragmentSelection();
+	FDelegateHandle OnFragmentSelectChangeHandle;
+};
diff --git a/Engine/Source/Editor/MassEntityDebugger/Private/SMassQuery.cpp b/Engine/Source/Editor/MassEntityDebugger/Private/SMassQuery.cpp
--- a/Engine/Source/Editor/MassEntityDebugger/Private/SMassQuery.cpp
+++ b/Engine/Source/Editor/MassEntityDebugger/Private/SMassQuery.cpp
@@ -16,7 +16,7 @@ static_assert(EMassAccessOperation::Read == uint32(EMassBitSetDisplayMode::ReadO
 namespace UE::Mass::Debugger::UI::Private
 {
 	template<typename TBitSet>
-	void AddBitSetPair(TSharedRef<SVerticalBox>& Box, const TMassExecutionAccess<TBitSet>& BitSetAccess, const FString& Label)
+	void AddBitSetPair(TSharedRef<SVerticalBox>& Box, const TMassExecutionAccess<TBitSet>& BitSetAccess, const FString& Label, TSharedRef<FMassDebuggerModel> InDebuggerModel)
 	{
 		if (BitSetAccess.IsEmpty() == false)
 		{
@@ -28,7 +28,7 @@ namespace UE::Mass::Debugger::UI::Private
 			Box->AddSlot()
 			.AutoHeight()
 			[
-				SNew(SMassBitSet<TBitSet>, Label, BitSetAccess.AsArrayView(), RequiredAccessBrushes)
+				SNew(SMassBitSet<TBitSet>, InDebuggerModel, Label, BitSetAccess.AsArrayView(), RequiredAccessBrushes)
 				.SlotPadding(5.0f)
 			];
 		}
@@ -36,7 +36,7 @@ namespace UE::Mass::Debugger::UI::Private
 } // namespace UE::Mass::Debugger::UI::Private
 
 
-void SMassQuery::Construct(const FArguments& InArgs, TSharedPtr<FMassDebuggerQueryData> InQueryData)
+void SMassQuery::Construct(const FArguments& InArgs, TSharedPtr<FMassDebuggerQueryData> InQueryData, TSharedRef<FMassDebuggerModel> InDebuggerModel)
 {
 	using UE::Mass::Debugger::UI::Private::AddBitSetPair;
 	using UE::Mass::Debugger::UI::AddBitSet;
@@ -60,14 +60,14 @@ void SMassQuery::Construct(const FArguments& InArgs, TSharedPtr<FMassDebuggerQue
 		.TextStyle(FAppStyle::Get(), "LargeText")
 	];
 
-	AddBitSetPair(Box, ExecutionRequirements.Fragments, TEXT("Fragments"));
-	AddBitSet(Box, ExecutionRequirements.RequiredAllTags, TEXT("RequiredAllTags"), ReadOnlyBrush);
-	AddBitSet(Box, ExecutionRequirements.RequiredAnyTags, TEXT("RequiredAnyTags"), ReadOnlyBrush);
-	AddBitSet(Box, ExecutionRequirements.RequiredNoneTags,TEXT("RequiredNoneTags"), ReadOnlyBrush);
-	AddBitSetPair(Box, ExecutionRequirements.ChunkFragments, TEXT("Chunk Fragments"));
-	AddBitSetPair(Box, ExecutionRequirements.SharedFragments, TEXT("Shared Fragments"));
-	AddBitSetPair(Box, ExecutionRequirements.ConstSharedFragments, TEXT("Const Shared Fragments"));
-	AddBitSetPair(Box, ExecutionRequirements.RequiredSubsystems, TEXT("Required Subsystems"));
+	AddBitSetPair(Box, ExecutionRequirements.Fragments, TEXT("Fragments"), InDebuggerModel);
+	AddBitSet(Box, ExecutionRequirements.RequiredAllTags, TEXT("RequiredAllTags"), ReadOnlyBrush, InDebuggerModel);
+	AddBitSet(Box, ExecutionRequirements.RequiredAnyTags, TEXT("RequiredAnyTags"), ReadOnlyBrush, InDebuggerModel);
+	AddBitSet(Box, ExecutionRequirements.RequiredNoneTags,TEXT("RequiredNoneTags"), ReadOnlyBrush, InDebuggerModel);
+	AddBitSetPair(Box, ExecutionRequirements.ChunkFragments, TEXT("Chunk Fragments"), InDebuggerModel);
+	AddBitSetPair(Box, ExecutionRequirements.SharedFragments, TEXT("Shared Fragments"), InDebuggerModel);
+	AddBitSetPair(Box, ExecutionRequirements.ConstSharedFragments, TEXT("Const Shared Fragments"), InDebuggerModel);
+	AddBitSetPair(Box, ExecutionRequirements.RequiredSubsystems, TEXT("Required Subsystems"), InDebuggerModel);
 
 	if (InQueryData->AdditionalInformation.IsEmpty() == false)
 	{
diff --git a/Engine/Source/Editor/MassEntityDebugger/Private/SMassQuery.h b/Engine/Source/Editor/MassEntityDebugger/Private/SMassQuery.h
--- a/Engine/Source/Editor/MassEntityDebugger/Private/SMassQuery.h
+++ b/Engine/Source/Editor/MassEntityDebugger/Private/SMassQuery.h
@@ -3,7 +3,7 @@
 
 #include "Widgets/SCompoundWidget.h"
 
-
+struct FMassDebuggerModel;
 struct FMassDebuggerQueryData;
 
 class SMassQuery : public SCompoundWidget
@@ -13,7 +13,7 @@ public:
 	SLATE_BEGIN_ARGS(SMassQuery){}
 	SLATE_END_ARGS()
 
-	void Construct(const FArguments& InArgs, TSharedPtr<FMassDebuggerQueryData> InQueryData);
+	void Construct(const FArguments& InArgs, TSharedPtr<FMassDebuggerQueryData> InQueryData, TSharedRef<FMassDebuggerModel> InDebuggerModel);
 
 protected:
 	TSharedPtr<FMassDebuggerQueryData> QueryData;
