[Mass]
Addes Breakpoints tab to Mass Debugger

diff --git a/Engine/Source/Editor/MassEntityDebugger/Private/SMassBreakpointsView.cpp b/Engine/Source/Editor/MassEntityDebugger/Private/SMassBreakpointsView.cpp
new file mode 100644
--- /dev/null
+++ b/Engine/Source/Editor/MassEntityDebugger/Private/SMassBreakpointsView.cpp
@@ -0,0 +1,56 @@
+// Copyright Epic Games, Inc. All Rights Reserved.
+
+#include "SMassBreakpointsView.h"
+#include "Widgets/Input/SButton.h"
+#include "Widgets/Layout/SScrollBox.h"
+#include "Widgets/Views/SListView.h"
+#include "Widgets/SBoxPanel.h"
+#include "MassDebugger.h"
+#include "MassDebuggerModel.h"
+
+#define LOCTEXT_NAMESPACE "SMassDebugger"
+
+
+//----------------------------------------------------------------------//
+// SMassBreakpointsView
+//----------------------------------------------------------------------//
+void SMassBreakpointsView::Construct(const FArguments& InArgs, TSharedRef<FMassDebuggerModel> InDebuggerModel)
+{
+	Initialize(InDebuggerModel);
+
+	ChildSlot
+	[
+		SNew(SVerticalBox)
+		+ SVerticalBox::Slot()
+		.AutoHeight()
+		.Padding(5)
+		[
+			SNew(SButton)
+			.Text(LOCTEXT("AddWriteBreakpoint", "Add Write Breakpoint"))
+			.OnClicked(this, &SMassBreakpointsView::HandleAddWriteBreakpointClicked)
+		]
+		+ SVerticalBox::Slot()
+		.AutoHeight()
+		.Padding(5)
+		[
+			SNew(SButton)
+			.Text(LOCTEXT("ClearAllBreakpoints", "Clear All Breakpoints"))
+			.OnClicked(this, &SMassBreakpointsView::ClearBreakpointsClicked)
+		]
+	];
+}
+
+FReply SMassBreakpointsView::HandleAddWriteBreakpointClicked()
+{
+	FMassDebugger::BreakOnFragmentWriteForSelectedEntity(DebuggerModel->GetSelectedFragment());
+	return FReply::Handled();
+}
+
+FReply SMassBreakpointsView::ClearBreakpointsClicked()
+{
+	FMassDebugger::ClearAllBreakpoints();
+	return FReply::Handled();
+}
+
+#undef LOCTEXT_NAMESPACE
+
diff --git a/Engine/Source/Editor/MassEntityDebugger/Private/SMassBreakpointsView.h b/Engine/Source/Editor/MassEntityDebugger/Private/SMassBreakpointsView.h
new file mode 100644
--- /dev/null
+++ b/Engine/Source/Editor/MassEntityDebugger/Private/SMassBreakpointsView.h
@@ -0,0 +1,42 @@
+// Copyright Epic Games, Inc. All Rights Reserved.
+#pragma once
+
+#include "SMassDebuggerViewBase.h"
+
+struct FMassDebuggerModel;
+struct FMassDebuggerArchetypeData;
+struct FMassDebuggerProcessorData;
+
+
+class SMassBreakpointsView : public SMassDebuggerViewBase
+{
+public:
+	SLATE_BEGIN_ARGS(SMassBreakpointsView)
+	{
+	}
+	SLATE_END_ARGS()
+
+	void Construct(const FArguments& InArgs, TSharedRef<FMassDebuggerModel> InDebuggerModel);
+
+protected:
+
+	/** Unused pure virtual function */
+	virtual void OnRefresh() override
+	{
+
+	}
+
+	/** Unused pure virtual function */
+	virtual void OnProcessorsSelected(TConstArrayView<TSharedPtr<FMassDebuggerProcessorData>> SelectedProcessors, ESelectInfo::Type SelectInfo) override
+	{
+	}
+
+	/** Unused pure virtual function */
+	virtual void OnArchetypesSelected(TConstArrayView<TSharedPtr<FMassDebuggerArchetypeData>> SelectedArchetypes, ESelectInfo::Type SelectInfo) override
+	{
+	}
+
+private:
+	FReply HandleAddWriteBreakpointClicked();
+	FReply ClearBreakpointsClicked();
+};
diff --git a/Engine/Source/Editor/MassEntityDebugger/Private/SMassDebugger.cpp b/Engine/Source/Editor/MassEntityDebugger/Private/SMassDebugger.cpp
--- a/Engine/Source/Editor/MassEntityDebugger/Private/SMassDebugger.cpp
+++ b/Engine/Source/Editor/MassEntityDebugger/Private/SMassDebugger.cpp
@@ -1,6 +1,7 @@
 // Copyright Epic Games, Inc. All Rights Reserved.
 
 #include "SMassDebugger.h"
+#include "SMassBreakpointsView.h"
 #include "SMassProcessorsView.h"
 #include "SMassProcessingView.h"
 #include "SMassArchetypesView.h"
@@ -31,6 +32,7 @@ void FMassDebuggerCommands::RegisterCommands()
 namespace UE::Mass::Debugger::Private
 {
 	const FName ToolbarTabId("Toolbar");
+	const FName BreakpointsTabId("Breakpoints");
 	const FName ProcessorsTabId("Processors");
 	const FName ProcessingGraphTabId("Processing Graphs");
 	const FName ArchetypesTabId("Archetypes");
@@ -101,6 +103,10 @@ void SMassDebugger::Construct(const FArguments& InArgs, const TSharedRef<SDockTa
 		.SetDisplayName(LOCTEXT("ArchetypesTabTitle", "Archetypes"))
 		.SetGroup(AppMenuGroup);
 
+	TabManager->RegisterTabSpawner(BreakpointsTabId, FOnSpawnTab::CreateRaw(this, &SMassDebugger::SpawnBreakpointsTab))
+		.SetDisplayName(LOCTEXT("BreakpointsTabTitle", "Breakpoints"))
+		.SetGroup(AppMenuGroup);
+
 
 	// Default Layout
 	TSharedRef<FTabManager::FLayout> Layout = FTabManager::NewLayout("MassDebuggerLayout_v1.0")
@@ -124,6 +130,7 @@ void SMassDebugger::Construct(const FArguments& InArgs, const TSharedRef<SDockTa
 						->AddTab(ProcessorsTabId, ETabState::OpenedTab)
 						->AddTab(ProcessingGraphTabId, ETabState::OpenedTab)
 						->AddTab(ArchetypesTabId, ETabState::OpenedTab)
+						->AddTab(BreakpointsTabId, ETabState::OpenedTab)
 						->SetForegroundTab(ProcessorsTabId)
 				)
 			)
@@ -319,6 +326,17 @@ TSharedRef<SDockTab> SMassDebugger::SpawnArchetypesTab(const FSpawnTabArgs& Args
 	return MajorTab;
 }
 
+TSharedRef<SDockTab> SMassDebugger::SpawnBreakpointsTab(const FSpawnTabArgs& Args)
+{
+	const TSharedRef<SDockTab> MajorTab = SNew(SDockTab)
+		.TabRole(ETabRole::PanelTab);
+
+	TSharedPtr<SWidget> TabContent = SNew(SMassBreakpointsView, DebuggerModel);
+	MajorTab->SetContent(TabContent.ToSharedRef());
+
+	return MajorTab;
+}
+
 void SMassDebugger::BindDelegates()
 {
 #if WITH_MASSENTITY_DEBUG
diff --git a/Engine/Source/Editor/MassEntityDebugger/Private/SMassDebugger.h b/Engine/Source/Editor/MassEntityDebugger/Private/SMassDebugger.h
--- a/Engine/Source/Editor/MassEntityDebugger/Private/SMassDebugger.h
+++ b/Engine/Source/Editor/MassEntityDebugger/Private/SMassDebugger.h
@@ -61,6 +61,7 @@ protected:
 	TSharedRef<SDockTab> SpawnProcessorsTab(const FSpawnTabArgs& Args);
 	TSharedRef<SDockTab> SpawnProcessingTab(const FSpawnTabArgs& Args);
 	TSharedRef<SDockTab> SpawnArchetypesTab(const FSpawnTabArgs& Args);
+	TSharedRef<SDockTab> SpawnBreakpointsTab(const FSpawnTabArgs& Args);
 
 	void BindDelegates();
 	void OnEntityManagerInitialized(const FMassEntityManager& EntityManager);
