[Mass]
Fix CIS 724436 and CIS  724444
Guard mass debugger editor code for cooked editor builds with entity debugging not enabled.

diff --git a/Engine/Source/Editor/MassEntityDebugger/Private/SMassArchetype.cpp b/Engine/Source/Editor/MassEntityDebugger/Private/SMassArchetype.cpp
--- a/Engine/Source/Editor/MassEntityDebugger/Private/SMassArchetype.cpp
+++ b/Engine/Source/Editor/MassEntityDebugger/Private/SMassArchetype.cpp
@@ -181,10 +181,12 @@ void SMassArchetype::Construct(const FArguments& InArgs, TSharedPtr<FMassDebugge
 
 void SMassArchetype::OnEntitySelected(TSharedPtr<FMassEntityHandle> SelectedItem, ESelectInfo::Type SelectInfo)
 {
+#if WITH_MASSENTITY_DEBUG
 	if (SelectedItem.IsValid() && DebuggerModel->Environment.IsValid())
 	{
 		FMassDebugger::SelectEntity(*(DebuggerModel->Environment->GetEntityManager()), *SelectedItem);
 	}
+#endif
 }
 
 TSharedRef<ITableRow> SMassArchetype::OnGenerateRow(TSharedPtr<FMassEntityHandle> InItem, const TSharedRef<STableViewBase>& OwnerTable)
@@ -192,6 +194,7 @@ TSharedRef<ITableRow> SMassArchetype::OnGenerateRow(TSharedPtr<FMassEntityHandle
 	return SNew(STableRow<TSharedPtr<FMassEntityHandle>>, OwnerTable)
 	.Content()
 	[
+#if WITH_MASSENTITY_DEBUG
 		SNew(SButton)
 		.Text(FText::FromString(InItem->DebugGetDescription()))
 		.OnClicked_Lambda([this, InItem]()
@@ -216,6 +219,10 @@ TSharedRef<ITableRow> SMassArchetype::OnGenerateRow(TSharedPtr<FMassEntityHandle
 				FMassDebugger::HighlightEntity(*(DebuggerModel->Environment->GetEntityManager()), FMassEntityHandle());
 			}
 		})
+#else
+		SNew(STextBlock)
+			.Text(LOCTEXT("MassEntityDebuggingNotEnabled", "Mass Entity Debugging Not Enabled for this configuration"))
+#endif
 	];
 }
 
diff --git a/Engine/Source/Editor/MassEntityDebugger/Private/SMassBreakpointsView.cpp b/Engine/Source/Editor/MassEntityDebugger/Private/SMassBreakpointsView.cpp
--- a/Engine/Source/Editor/MassEntityDebugger/Private/SMassBreakpointsView.cpp
+++ b/Engine/Source/Editor/MassEntityDebugger/Private/SMassBreakpointsView.cpp
@@ -42,13 +42,17 @@ void SMassBreakpointsView::Construct(const FArguments& InArgs, TSharedRef<FMassD
 
 FReply SMassBreakpointsView::HandleAddWriteBreakpointClicked()
 {
-	FMassDebugger::BreakOnFragmentWriteForSelectedEntity(DebuggerModel->GetSelectedFragment());
+#if WITH_MASSENTITY_DEBUG
+	FMassDebugger::BreakOnFragmentWriteForSelectedEntity(DebuggerModel->GetSelectedFragment()); 
+#endif
 	return FReply::Handled();
 }
 
 FReply SMassBreakpointsView::ClearBreakpointsClicked()
 {
+#if WITH_MASSENTITY_DEBUG
 	FMassDebugger::ClearAllBreakpoints();
+#endif
 	return FReply::Handled();
 }
 
