[Mass]
Fix CIS 724463
Guard mass debugger editor code for cooked editor builds with entity debugging not enabled.

diff --git a/Engine/Source/Editor/MassEntityDebugger/Private/SMassProcessorsWidget.cpp b/Engine/Source/Editor/MassEntityDebugger/Private/SMassProcessorsWidget.cpp
--- a/Engine/Source/Editor/MassEntityDebugger/Private/SMassProcessorsWidget.cpp
+++ b/Engine/Source/Editor/MassEntityDebugger/Private/SMassProcessorsWidget.cpp
@@ -234,6 +234,7 @@ const FSlateBrush* SMassProcessorWidget::GetBorderByFragmentSelection()
 	}
 
 	auto HasSelectedBit = [&SelectedFragment](const auto& Bitset) -> bool {
+#if WITH_MASSENTITY_DEBUG
 		for (int32 Index = 0; Index < Bitset.StructTypesBitArray.Num(); ++Index)
 		{
 			if (Bitset.StructTypesBitArray[Index])
@@ -245,6 +246,7 @@ const FSlateBrush* SMassProcessorWidget::GetBorderByFragmentSelection()
 				}
 			}
 		}
+#endif
 		return false;
 		};
 
