[Mass]
Fix 724517
Added missing #undef LOCTEXT_NAMESPACE

diff --git a/Engine/Source/Editor/MassEntityDebugger/Private/SMassArchetype.cpp b/Engine/Source/Editor/MassEntityDebugger/Private/SMassArchetype.cpp
--- a/Engine/Source/Editor/MassEntityDebugger/Private/SMassArchetype.cpp
+++ b/Engine/Source/Editor/MassEntityDebugger/Private/SMassArchetype.cpp
@@ -226,3 +226,4 @@ TSharedRef<ITableRow> SMassArchetype::OnGenerateRow(TSharedPtr<FMassEntityHandle
 	];
 }
 
+#undef LOCTEXT_NAMESPACE
\ No newline at end of file
diff --git a/Engine/Source/Editor/MassEntityDebugger/Private/SMassProcessorsWidget.cpp b/Engine/Source/Editor/MassEntityDebugger/Private/SMassProcessorsWidget.cpp
--- a/Engine/Source/Editor/MassEntityDebugger/Private/SMassProcessorsWidget.cpp
+++ b/Engine/Source/Editor/MassEntityDebugger/Private/SMassProcessorsWidget.cpp
@@ -295,3 +295,4 @@ const FSlateBrush* SMassProcessorWidget::GetBorderByFragmentSelection()
 	return AccessColor[HighestAccess];
 }
 
+#undef LOCTEXT_NAMESPACE
\ No newline at end of file
