[Mass] Changed how composite/phase processors get collected by the Mass Debugger

The change supports having multiple independent (entity manager, phase manager) pairs per given runtime UWorld.

diff --git a/Engine/Source/Editor/MassEntityDebugger/Private/MassDebuggerModel.cpp b/Engine/Source/Editor/MassEntityDebugger/Private/MassDebuggerModel.cpp
--- a/Engine/Source/Editor/MassEntityDebugger/Private/MassDebuggerModel.cpp
+++ b/Engine/Source/Editor/MassEntityDebugger/Private/MassDebuggerModel.cpp
@@ -340,11 +340,11 @@ FText FMassDebuggerProcessingGraphNode::GetLabel() const
 //----------------------------------------------------------------------//
 // FMassDebuggerProcessingGraph
 //----------------------------------------------------------------------//
-FMassDebuggerProcessingGraph::FMassDebuggerProcessingGraph(const FMassDebuggerModel& DebuggerModel, UMassCompositeProcessor& InGraphOwner)
+FMassDebuggerProcessingGraph::FMassDebuggerProcessingGraph(const FMassDebuggerModel& DebuggerModel, TNotNull<const UMassCompositeProcessor*> InGraphOwner)
 {
-	Label = InGraphOwner.GetProcessorName();
+	Label = InGraphOwner->GetProcessorName();
 #if WITH_MASSENTITY_DEBUG
-	TConstArrayView<UMassCompositeProcessor::FDependencyNode> ProcessingGraph = FMassDebugger::GetProcessingGraph(InGraphOwner);
+	TConstArrayView<UMassCompositeProcessor::FDependencyNode> ProcessingGraph = FMassDebugger::GetProcessingGraph(*InGraphOwner);
 
 	if (ProcessingGraph.Num() > 0)
 	{
@@ -359,9 +359,9 @@ FMassDebuggerProcessingGraph::FMassDebuggerProcessingGraph(const FMassDebuggerMo
 	}
 	// it's possible for the graph to be empty if InGraphOwner has been populated for a single-thread execution.
 	// See if there are any processors owned by InGraphOwner.
-	else if (InGraphOwner.IsEmpty() == false)
+	else if (InGraphOwner->IsEmpty() == false)
 	{
-		TConstArrayView<TObjectPtr<UMassProcessor>> HostedProcessors = FMassDebugger::GetHostedProcessors(InGraphOwner);
+		TConstArrayView<TObjectPtr<UMassProcessor>> HostedProcessors = FMassDebugger::GetHostedProcessors(*InGraphOwner);
 		for (const TObjectPtr<UMassProcessor>& Processor : HostedProcessors)
 		{
 			check(Processor);
@@ -431,8 +431,9 @@ void FMassDebuggerModel::RefreshAll()
 		TMap<FMassArchetypeHandle, TSharedPtr<FMassDebuggerArchetypeData>> TransientArchetypesMap;
 
 		CacheArchetypesData(TransientArchetypesMap);
-		CacheProcessorsData(TransientArchetypesMap);
-		CacheProcessingGraphs();
+		TArray<TNotNull<const UMassCompositeProcessor*>> CompositeProcessors;
+		CacheProcessorsData(TransientArchetypesMap, CompositeProcessors);
+		CacheProcessingGraphs(CompositeProcessors);
 
 		ClearArchetypeSelection();
 
@@ -512,7 +513,7 @@ void FMassDebuggerModel::ClearArchetypeSelection()
 	OnArchetypesSelectedDelegate.Broadcast(SelectedArchetypes, ESelectInfo::Direct);
 }
 
-void FMassDebuggerModel::CacheProcessorsData(const TMap<FMassArchetypeHandle, TSharedPtr<FMassDebuggerArchetypeData>>& InTransientArchetypesMap)
+void FMassDebuggerModel::CacheProcessorsData(const TMap<FMassArchetypeHandle, TSharedPtr<FMassDebuggerArchetypeData>>& InTransientArchetypesMap, TArray<TNotNull<const UMassCompositeProcessor*>>& OutCompositeProcessors)
 {
 	static auto SortPredicate = [](const TSharedPtr<FMassDebuggerProcessorData>& A, const TSharedPtr<FMassDebuggerProcessorData>& B)
 	{
@@ -551,9 +552,19 @@ void FMassDebuggerModel::CacheProcessorsData(const TMap<FMassArchetypeHandle, TS
 			{
 				if (Processor)
 				{
-					AllCachedProcessors.Add(
-						Container.Add_GetRef(MakeShareable(new FMassDebuggerProcessorData(*EntityManager, *Processor, InTransientArchetypesMap)))
-					);
+					if (const UMassCompositeProcessor* CompositeProcessor = Cast<const UMassCompositeProcessor>(Processor))
+					{
+						// The composite processors are collected in a dedicated container to get processed
+						// in a special way. These processors will be shown in "processing phase" tab.
+						// @todo rename the mentioned tab to indicate that it contains all provided composite processors.
+						OutCompositeProcessors.Emplace(CompositeProcessor);
+					}
+					else
+					{
+						AllCachedProcessors.Add(
+							Container.Add_GetRef(MakeShareable(new FMassDebuggerProcessorData(*EntityManager, *Processor, InTransientArchetypesMap)))
+							);
+					}
 				}
 			}
 			Container.Sort(SortPredicate);
@@ -581,19 +592,14 @@ void FMassDebuggerModel::CacheProcessorsData(const TMap<FMassArchetypeHandle, TS
 	AllCachedProcessors.Sort(SortPredicate);
 }
 
-void FMassDebuggerModel::CacheProcessingGraphs()
+void FMassDebuggerModel::CacheProcessingGraphs(TConstArrayView<TNotNull<const UMassCompositeProcessor*>> InCompositeProcessors)
 {
 	CachedProcessingGraphs.Reset();
 
-	UWorld* World = Environment ? Environment->World.Get() : nullptr;
-	for (FThreadSafeObjectIterator It(UMassCompositeProcessor::StaticClass()); It; ++It)
+	for (TNotNull<const UMassCompositeProcessor*> Processor : InCompositeProcessors)
 	{
-		UMassCompositeProcessor* Processor = CastChecked<UMassCompositeProcessor>(*It);
-		if (UE::Mass::Debugger::Private::IsDebuggableProcessor(World, *Processor))
-		{
-			CachedProcessingGraphs.Add(MakeShareable(new FMassDebuggerProcessingGraph(*this, *Processor)));
-		}
-	}	
+		CachedProcessingGraphs.Add(MakeShareable(new FMassDebuggerProcessingGraph(*this, Processor)));
+	}
 }
 
 void FMassDebuggerModel::CacheArchetypesData(TMap<FMassArchetypeHandle, TSharedPtr<FMassDebuggerArchetypeData>>& OutTransientArchetypesMap)
diff --git a/Engine/Source/Editor/MassEntityDebugger/Private/MassDebuggerModel.h b/Engine/Source/Editor/MassEntityDebugger/Private/MassDebuggerModel.h
--- a/Engine/Source/Editor/MassEntityDebugger/Private/MassDebuggerModel.h
+++ b/Engine/Source/Editor/MassEntityDebugger/Private/MassDebuggerModel.h
@@ -136,7 +136,7 @@ struct FMassDebuggerProcessingGraphNode
 
 struct FMassDebuggerProcessingGraph
 {
-	FMassDebuggerProcessingGraph(const FMassDebuggerModel& DebuggerModel, UMassCompositeProcessor& InGraphOwner);
+	FMassDebuggerProcessingGraph(const FMassDebuggerModel& DebuggerModel, TNotNull<const UMassCompositeProcessor*> InGraphOwner);
 
 	FString Label;
 	TArray<FMassDebuggerProcessingGraphNode> GraphNodes;
@@ -188,8 +188,8 @@ struct FMassDebuggerModel
 	bool HasEnvironmentSelected() const { return static_cast<bool>(Environment); }
 
 	void CacheArchetypesData(TMap<FMassArchetypeHandle, TSharedPtr<FMassDebuggerArchetypeData>>& OutTransientArchetypesMap); 
-	void CacheProcessorsData(const TMap<FMassArchetypeHandle, TSharedPtr<FMassDebuggerArchetypeData>>& InTransientArchetypesMap);
-	void CacheProcessingGraphs();
+	void CacheProcessorsData(const TMap<FMassArchetypeHandle, TSharedPtr<FMassDebuggerArchetypeData>>& InTransientArchetypesMap, TArray<TNotNull<const UMassCompositeProcessor*>>& OutCompositeProcessors);
+	void CacheProcessingGraphs(TConstArrayView<TNotNull<const UMassCompositeProcessor*>> InCompositeProcessors);
 
 	float MinDistanceToSelectedArchetypes(const TSharedPtr<FMassDebuggerArchetypeData>& InArchetypeData) const;
 
diff --git a/Engine/Source/Runtime/MassEntity/Private/MassProcessingPhaseManager.cpp b/Engine/Source/Runtime/MassEntity/Private/MassProcessingPhaseManager.cpp
--- a/Engine/Source/Runtime/MassEntity/Private/MassProcessingPhaseManager.cpp
+++ b/Engine/Source/Runtime/MassEntity/Private/MassProcessingPhaseManager.cpp
@@ -334,6 +334,7 @@ void FMassProcessingPhaseManager::Start(const TSharedRef<FMassEntityManager>& In
 		{
 			for (const FMassProcessingPhase& Phase : SharedThis->ProcessingPhases)
 			{
+				OutProcessors.Add(Phase.DebugGetPhaseProcessor());
 				OutProcessors.Append(Phase.DebugGetPhaseProcessor()->GetChildProcessorsView());
 			}
 		}
