[Mass Debugger]
Limit the number of entites shown in the Archetypes entity list.

diff --git a/Engine/Source/Editor/MassEntityDebugger/Private/SMassArchetype.cpp b/Engine/Source/Editor/MassEntityDebugger/Private/SMassArchetype.cpp
--- a/Engine/Source/Editor/MassEntityDebugger/Private/SMassArchetype.cpp
+++ b/Engine/Source/Editor/MassEntityDebugger/Private/SMassArchetype.cpp
@@ -25,10 +25,14 @@ void SMassArchetype::Construct(const FArguments& InArgs, TSharedPtr<FMassDebugge
 	ArchetypeData = InArchetypeData;
 	DebuggerModel = InDebuggerModel;
 
-	EntityHandles.Reserve(ArchetypeData->EntityHandles.Num());
-	for (FMassEntityHandle& Handle : ArchetypeData->EntityHandles)
+	// limit the number of entities shown to avoid perf issues
+	// @todo replace this temp fix with virtualized scrolling of the entire entity list (and filtering)
+	constexpr int MaxEntitiesToDisplay = 1000;
+	const int EntityDisplayCount = FMath::Min(ArchetypeData->EntityHandles.Num(), MaxEntitiesToDisplay);
+	EntityHandles.Reserve(EntityDisplayCount);
+	for (int i = 0; i < EntityDisplayCount; i++)
 	{
-		EntityHandles.Add(MakeShared<FMassEntityHandle>(Handle));
+		EntityHandles.Add(MakeShared<FMassEntityHandle>(ArchetypeData->EntityHandles[i]));
 	}
 
 	const FMassDebuggerArchetypeData* BaseArchetypeDebugData = InBaseArchetypeData.Get();
