[Mass Debugger]
Properly handle editor builds without mass entity debugging enabled.

diff --git a/Engine/Source/Editor/MassEntityDebugger/Private/SMassBreakpointsView.cpp b/Engine/Source/Editor/MassEntityDebugger/Private/SMassBreakpointsView.cpp
--- a/Engine/Source/Editor/MassEntityDebugger/Private/SMassBreakpointsView.cpp
+++ b/Engine/Source/Editor/MassEntityDebugger/Private/SMassBreakpointsView.cpp
@@ -16,6 +16,7 @@
 //----------------------------------------------------------------------//
 void SMassBreakpointsView::Construct(const FArguments& InArgs, TSharedRef<FMassDebuggerModel> InDebuggerModel)
 {
+#if WITH_MASSENTITY_DEBUG
 	Initialize(InDebuggerModel);
 
 	FMassDebugger::OnBreakpointsChangedDelegate.AddSP(this, &SMassBreakpointsView::RefreshBreakpoints);
@@ -50,10 +51,18 @@ void SMassBreakpointsView::Construct(const FArguments& InArgs, TSharedRef<FMassD
 			.OnClicked(this, &SMassBreakpointsView::ClearBreakpointsClicked)
 		]
 	];
+#else
+	ChildSlot
+	[
+		SNew(STextBlock)
+			.Text(LOCTEXT("MassEntityDebuggingNotEnabled", "Mass Entity Debugging Not Enabled for this configuration"))
+	];
+#endif
 }
 
 void SMassBreakpointsView::RefreshBreakpoints()
 {
+#if WITH_MASSENTITY_DEBUG
 	WriteBreakpoints.Reset();
 	if (DebuggerModel && DebuggerModel->Environment)
 	{
@@ -71,6 +80,7 @@ void SMassBreakpointsView::RefreshBreakpoints()
 	{
 		WriteBreakpointsListView->RebuildList();
 	}
+#endif
 }
 
 FReply SMassBreakpointsView::HandleAddWriteBreakpointClicked()
