[Mass Debugger]
Added window menu to recover lost tabs and reset the layout

diff --git a/Engine/Source/Editor/MassEntityDebugger/Private/SMassDebugger.cpp b/Engine/Source/Editor/MassEntityDebugger/Private/SMassDebugger.cpp
--- a/Engine/Source/Editor/MassEntityDebugger/Private/SMassDebugger.cpp
+++ b/Engine/Source/Editor/MassEntityDebugger/Private/SMassDebugger.cpp
@@ -74,14 +74,8 @@ void SMassDebugger::Construct(const FArguments& InArgs, const TSharedRef<SDockTa
 	FUICommandList& ActionList = *CommandList;
 
 	ActionList.MapAction(Commands.RefreshData
-		, FExecuteAction::CreateLambda([this]() 
-		{
-				DebuggerModel->RefreshAll();
-		})
-		// check if there's any mass environment picked
-		, FCanExecuteAction::CreateLambda([this]()->bool { return DebuggerModel->HasEnvironmentSelected(); })
-		, FIsActionChecked::CreateLambda([]()->bool { return false; })
-		, FIsActionButtonVisible::CreateLambda([this]()->bool { return true; }));
+		, FExecuteAction::CreateSP(this, &SMassDebugger::RefreshData)
+		, FCanExecuteAction::CreateSP(this, &SMassDebugger::CanRefreshData));
 
 	// Tab Spawners
 	TabManager = FGlobalTabmanager::Get()->NewTabManager(ConstructUnderMajorTab);
@@ -109,7 +103,32 @@ void SMassDebugger::Construct(const FArguments& InArgs, const TSharedRef<SDockTa
 
 
 	// Default Layout
-	TSharedRef<FTabManager::FLayout> Layout = FTabManager::NewLayout("MassDebuggerLayout_v1.0")
+	TSharedRef<FTabManager::FLayout> Layout = CreateDefaultLayout();
+
+	Layout = FLayoutSaveRestore::LoadFromConfig(GEditorLayoutIni, Layout);
+
+	ChildSlot
+	[
+		TabManager->RestoreFrom(Layout, ConstructUnderWindow).ToSharedRef()
+	];
+
+	TabManager->SetOnPersistLayout(
+		FTabManager::FOnPersistLayout::CreateStatic(
+			[](const TSharedRef<FTabManager::FLayout>& InLayout)
+			{
+				if (InLayout->GetPrimaryArea().Pin().IsValid())
+				{
+					FLayoutSaveRestore::SaveToConfig(GEditorLayoutIni, InLayout);
+				}
+			}
+		)
+	);
+}
+
+TSharedRef<FTabManager::FLayout> SMassDebugger::CreateDefaultLayout()
+{
+	using namespace UE::Mass::Debugger::Private;
+	TSharedRef<FTabManager::FLayout> Layout = FTabManager::NewLayout("MassDebuggerLayout_v1.2")
 		->AddArea
 		(
 			FTabManager::NewPrimaryArea()
@@ -126,34 +145,30 @@ void SMassDebugger::Construct(const FArguments& InArgs, const TSharedRef<SDockTa
 				->SetOrientation(Orient_Horizontal)
 				->Split
 				(
-					FTabManager::NewStack()
+					FTabManager::NewSplitter()
+					->SetOrientation(Orient_Vertical)
+					->Split
+					(
+						FTabManager::NewStack()
 						->AddTab(ProcessorsTabId, ETabState::OpenedTab)
 						->AddTab(ProcessingGraphTabId, ETabState::OpenedTab)
-						->AddTab(ArchetypesTabId, ETabState::OpenedTab)
+						->SetForegroundTab(ProcessorsTabId)
+					)
+					->Split
+					(
+						FTabManager::NewStack()
 						->AddTab(BreakpointsTabId, ETabState::OpenedTab)
 						->SetForegroundTab(ProcessorsTabId)
+					)
+				)
+				->Split
+				(
+					FTabManager::NewStack()
+					->AddTab(ArchetypesTabId, ETabState::OpenedTab)
 				)
 			)
 		);
-
-	Layout = FLayoutSaveRestore::LoadFromConfig(GEditorLayoutIni, Layout);
-
-	ChildSlot
-		[
-			TabManager->RestoreFrom(Layout, ConstructUnderWindow).ToSharedRef()
-		];
-
-	TabManager->SetOnPersistLayout(
-		FTabManager::FOnPersistLayout::CreateStatic(
-			[](const TSharedRef<FTabManager::FLayout>& InLayout)
-			{
-				if (InLayout->GetPrimaryArea().Pin().IsValid())
-				{
-					FLayoutSaveRestore::SaveToConfig(GEditorLayoutIni, InLayout);
-				}
-			}
-		)
-	);
+	return Layout;
 }
 
 TSharedRef<SDockTab> SMassDebugger::SpawnToolbar(const FSpawnTabArgs& Args)
@@ -163,7 +178,19 @@ TSharedRef<SDockTab> SMassDebugger::SpawnToolbar(const FSpawnTabArgs& Args)
 		.ShouldAutosize(true);
 
 	FSlimHorizontalToolBarBuilder ToolBarBuilder(CommandList, FMultiBoxCustomization::None);
-	ToolBarBuilder.BeginSection("Debugger");
+	ToolBarBuilder.BeginSection(TEXT("Window"));
+	{
+		ToolBarBuilder.AddComboButton(
+			FUIAction()
+			, FOnGetContent::CreateSP(this, &SMassDebugger::GenerateWindowMenu)
+			, LOCTEXT("MassDebuggerWindowMenu_Label", "Window")
+			, LOCTEXT("MassWindowMenu_Tooltip", "Window options")
+			, FSlateIcon()
+			, false
+		);
+	}
+	ToolBarBuilder.EndSection();
+	ToolBarBuilder.BeginSection(TEXT("Debugger"));
 	{
 		ToolBarBuilder.AddToolBarButton(FMassDebuggerCommands::Get().RefreshData, NAME_None, LOCTEXT("RefreshData", "Refresh"), LOCTEXT("RefreshDebuggerTooltip", "Refreshes data cached by the debugger instance"));//, FSlateIcon(FLogVisualizerStyle::Get().GetStyleSetName(), TEXT("Toolbar.Update")));
 	}
@@ -201,6 +228,27 @@ TSharedRef<SDockTab> SMassDebugger::SpawnToolbar(const FSpawnTabArgs& Args)
 	return MajorTab;
 }
 
+TSharedRef<SWidget> SMassDebugger::GenerateWindowMenu()
+{
+	FMenuBuilder MenuBuilder(true, nullptr);
+
+	FUIAction ProcessorsUIAction(FExecuteAction::CreateSP(this, &SMassDebugger::ShowProcessorView));
+	FUIAction ArchetypesUIAction(FExecuteAction::CreateSP(this, &SMassDebugger::ShowArchetypesView));
+	FUIAction ProcessingGraphsUIAction(FExecuteAction::CreateSP(this, &SMassDebugger::ShowProcessingGraphsView));
+	FUIAction BreakpointsUIAction(FExecuteAction::CreateSP(this, &SMassDebugger::ShowBreakpointsView));
+	FUIAction ResetLayoutUIAction(FExecuteAction::CreateSP(this, &SMassDebugger::ResetLayout));
+
+	MenuBuilder.BeginSection(TEXT("Tabs"));
+	MenuBuilder.AddMenuEntry(LOCTEXT("ProcessorsTabLabel", "Processors"), LOCTEXT("ProcessorsTabTooltip", "Show Processors Tab"), FSlateIcon(), ProcessorsUIAction);
+	MenuBuilder.AddMenuEntry(LOCTEXT("ArchetypesTabLabel", "Archetypes"), LOCTEXT("ArchtypesTabTooltip", "Show Archtypes Tab"), FSlateIcon(), ArchetypesUIAction);
+	MenuBuilder.AddMenuEntry(LOCTEXT("ProcessingGraphsTabLabel", "Processing Graphs"), LOCTEXT("ProcessingGraphsTabTooltip", "Show Processing Graphs Tab"), FSlateIcon(), ProcessingGraphsUIAction);
+	MenuBuilder.AddMenuEntry(LOCTEXT("BreakpointsTabLabel", "Breakpoints"), LOCTEXT("BreakpointsTabTooltip", "Show Breakpoints Tab"), FSlateIcon(), BreakpointsUIAction);
+	MenuBuilder.AddSeparator();
+	MenuBuilder.AddMenuEntry(LOCTEXT("ResetLayoutLabel", "Reset Layout"), LOCTEXT("ResetLayoutTooltip", "Reset the Mass Debugger Layout"), FSlateIcon(), ResetLayoutUIAction);
+
+	return MenuBuilder.MakeWidget();
+}
+
 void SMassDebugger::OnEntityManagerInitialized(const FMassEntityManager& EntityManager)
 {
 	const UWorld* World = EntityManager.GetWorld();
@@ -337,6 +385,52 @@ TSharedRef<SDockTab> SMassDebugger::SpawnBreakpointsTab(const FSpawnTabArgs& Arg
 	return MajorTab;
 }
 
+bool SMassDebugger::CanRefreshData()
+{
+	return DebuggerModel->HasEnvironmentSelected();
+}
+
+void SMassDebugger::RefreshData()
+{
+	DebuggerModel->RefreshAll();
+}
+
+void SMassDebugger::ShowProcessorView() const
+{
+	TabManager->TryInvokeTab(UE::Mass::Debugger::Private::ProcessorsTabId);
+}
+
+void SMassDebugger::ShowArchetypesView() const
+{
+	TabManager->TryInvokeTab(UE::Mass::Debugger::Private::ArchetypesTabId);
+}
+
+void SMassDebugger::ShowBreakpointsView() const
+{
+	TabManager->TryInvokeTab(UE::Mass::Debugger::Private::BreakpointsTabId);
+}
+
+void SMassDebugger::ShowProcessingGraphsView() const
+{
+	TabManager->TryInvokeTab(UE::Mass::Debugger::Private::ProcessingGraphTabId);
+}
+
+void SMassDebugger::ResetLayout()
+{
+	TSharedPtr<SWindow> Window = FSlateApplication::Get().FindWidgetWindow(AsShared());
+	if (Window.IsValid())
+	{
+		ChildSlot.DetachWidget();
+		ChildSlot.AttachWidget(TabManager->RestoreFrom(CreateDefaultLayout(), Window).ToSharedRef());
+	}
+}
+
+void SMassDebugger::ShowSelectedView() const
+{
+	
+}
+
+
 void SMassDebugger::BindDelegates()
 {
 #if WITH_MASSENTITY_DEBUG
diff --git a/Engine/Source/Editor/MassEntityDebugger/Private/SMassDebugger.h b/Engine/Source/Editor/MassEntityDebugger/Private/SMassDebugger.h
--- a/Engine/Source/Editor/MassEntityDebugger/Private/SMassDebugger.h
+++ b/Engine/Source/Editor/MassEntityDebugger/Private/SMassDebugger.h
@@ -63,12 +63,25 @@ protected:
 	TSharedRef<SDockTab> SpawnArchetypesTab(const FSpawnTabArgs& Args);
 	TSharedRef<SDockTab> SpawnBreakpointsTab(const FSpawnTabArgs& Args);
 
+	TSharedRef<FTabManager::FLayout> CreateDefaultLayout();
+	TSharedRef<SWidget> GenerateWindowMenu();
+
 	void BindDelegates();
 	void OnEntityManagerInitialized(const FMassEntityManager& EntityManager);
 	void OnEntityManagerDeinitialized(const FMassEntityManager& EntityManager);
 	void HandleEnvironmentChanged(TSharedPtr<FMassDebuggerEnvironment> Item, ESelectInfo::Type SelectInfo);
 	void RebuildEnvironmentsList();
 
+	bool CanRefreshData();
+	void RefreshData();
+
+	void ShowProcessorView() const;
+	void ShowArchetypesView() const;
+	void ShowBreakpointsView() const;
+	void ShowProcessingGraphsView() const;
+	void ResetLayout();
+	void ShowSelectedView() const;
+
 #if WITH_MASSENTITY_DEBUG
 void OnProcessorProviderRegistered(const FMassDebugger::FEnvironment& Environment);
 #endif // WITH_MASSENTITY_DEBUG
