[Mass Debugger]
Added entites view
Added fragment data inspection UI
Improved fragment breakpoint workflow

diff --git a/Engine/Source/Editor/MassEntityDebugger/Private/MassDebuggerModel.cpp b/Engine/Source/Editor/MassEntityDebugger/Private/MassDebuggerModel.cpp
--- a/Engine/Source/Editor/MassEntityDebugger/Private/MassDebuggerModel.cpp
+++ b/Engine/Source/Editor/MassEntityDebugger/Private/MassDebuggerModel.cpp
@@ -12,6 +12,8 @@
 #include "Containers/UnrealString.h"
 #include "MassArchetypeData.h"
 #include "Engine/World.h"
+#include "SMassEntitiesView.h"
+#include "SMassDebugger.h"
 
 #define LOCTEXT_NAMESPACE "SMassDebugger"
 
@@ -227,6 +229,7 @@ void FMassDebuggerProcessorData::SetProcessor(const UMassProcessor& InProcessor)
 FMassDebuggerArchetypeData::FMassDebuggerArchetypeData(const FMassArchetypeHandle& ArchetypeHandle)
 {
 #if WITH_MASSENTITY_DEBUG
+	Handle = ArchetypeHandle;
 	Composition = FMassDebugger::GetArchetypeComposition(ArchetypeHandle);
 
 	// @todo should ensure we're using same hashing as the EntityManager here
@@ -238,7 +241,6 @@ FMassDebuggerArchetypeData::FMassDebuggerArchetypeData(const FMassArchetypeHandl
 	HashLabel = FText::FromString(FullHashAsString);
 
 	FMassDebugger::GetArchetypeEntityStats(ArchetypeHandle, ArchetypeStats);
-	EntityHandles = FMassDebugger::GetEntitiesOfArchetype(ArchetypeHandle);
 
 	const TConstArrayView<FName> DebugNames = FMassDebugger::GetArchetypeDebugNames(ArchetypeHandle);
 
@@ -731,6 +733,70 @@ const TSharedPtr<FMassDebuggerProcessorData>& FMassDebuggerModel::GetProcessorDa
 	return *DataFound;
 }
 
+void FMassDebuggerModel::RegisterEntitiesView(TSharedRef<SMassEntitiesView> EntitiesView, int32 Index)
+{
+	if (EntityViews.Num() < (Index + 1))
+	{
+		EntityViews.SetNum(Index + 1);
+	}
+	EntityViews[Index] = EntitiesView;
+}
+
+void FMassDebuggerModel::ShowEntitiesView(int Index, FMassArchetypeHandle ArchetypeHandle)
+{
+	TWeakPtr<SMassEntitiesView> View = ShowEntitiesView(Index);
+	if (!View.IsValid())
+	{
+		return;
+	}
+	View.Pin()->ShowEntities(ArchetypeHandle);
+}
+
+void FMassDebuggerModel::ShowEntitiesView(int Index, TArray<FMassEntityHandle> EntitieHandles)
+{
+	TWeakPtr<SMassEntitiesView> View = ShowEntitiesView(Index);
+	if (!View.IsValid())
+	{
+		return;
+	}
+	View.Pin()->ShowEntities(EntitieHandles);
+}
+
+void FMassDebuggerModel::ShowEntitiesView(int Index, FMassEntityQuery& Query)
+{
+	TWeakPtr<SMassEntitiesView> View = ShowEntitiesView(Index);
+	if (!View.IsValid())
+	{
+		return;
+	}
+	View.Pin()->ShowEntities(Query);
+}
+
+void FMassDebuggerModel::ShowEntitiesView(int Index, TConstArrayView<FMassEntityQuery*> InQueries)
+{
+	TWeakPtr<SMassEntitiesView> View = ShowEntitiesView(Index);
+	if (!View.IsValid())
+	{
+		return;
+	}
+	View.Pin()->ShowEntities(InQueries);
+}
+
+TWeakPtr<SMassEntitiesView> FMassDebuggerModel::ShowEntitiesView(int32 Index)
+{
+	if (DebuggerWindow.IsValid())
+	{
+		DebuggerWindow.Pin()->ShowEntitesView();
+	}
+
+	check(Index < MaxEntityViewCount);
+	if (EntityViews.Num() < (Index + 1))
+	{
+		// TODO: create the tab and set focus to it
+	}
+	return EntityViews[Index];
+}
+
 void FMassDebuggerModel::ResetSelectedArchetypes()
 {
 	for (TSharedPtr<FMassDebuggerArchetypeData>& ArchetypeData : SelectedArchetypes)
diff --git a/Engine/Source/Editor/MassEntityDebugger/Private/MassDebuggerModel.h b/Engine/Source/Editor/MassEntityDebugger/Private/MassDebuggerModel.h
--- a/Engine/Source/Editor/MassEntityDebugger/Private/MassDebuggerModel.h
+++ b/Engine/Source/Editor/MassEntityDebugger/Private/MassDebuggerModel.h
@@ -15,7 +15,10 @@ struct FMassArchetypeHandle;
 struct FMassDebuggerModel;
 struct FMassEntityHandle;
 struct FMassEntityManager;
+struct FMassEntityQuery;
 class UWorld;
+class SMassEntitiesView;
+class SMassDebugger;
 
 enum class EMassDebuggerSelectionMode : uint8
 {
@@ -60,6 +63,8 @@ struct FMassDebuggerArchetypeData
 {
 	FMassDebuggerArchetypeData(const FMassArchetypeHandle& ArchetypeHandle);
 
+	FMassArchetypeHandle Handle;
+
 	FMassArchetypeCompositionDescriptor Composition;
 
 	/** Hash of the Compositions. */
@@ -89,9 +94,6 @@ struct FMassDebuggerArchetypeData
 	/** Primary debug name, used for grouping derived archetypes. */
 	FString PrimaryDebugName;
 
-	/** Snapshot of all entity handles using this archetype when this struct was created */
-	TArray<FMassEntityHandle> EntityHandles;
-
 	/** True if the archetype is selected. */
 	bool bIsSelected = false;
 
@@ -200,6 +202,15 @@ struct FMassDebuggerModel
 
 	const TSharedPtr<FMassDebuggerProcessorData>& GetProcessorDataChecked(const UMassProcessor& Processor) const;
 
+	void RegisterEntitiesView(TSharedRef<SMassEntitiesView> EntitiesView, int32 Index);
+
+	static const int32 MaxEntityViewCount = 1;
+	void ShowEntitiesView(int Index, FMassArchetypeHandle ArchetypeHandle);
+	void ShowEntitiesView(int Index, TArray<FMassEntityHandle> EntitieHandles);
+	void ShowEntitiesView(int Index, FMassEntityQuery& Query);
+	void ShowEntitiesView(int Index, TConstArrayView<FMassEntityQuery*> InQueries);
+	TWeakPtr<SMassEntitiesView> ShowEntitiesView(int32 Index);
+
 protected:
 	void StoreArchetypes(const FMassEntityManager& EntitySubsystem, TMap<FMassArchetypeHandle, TSharedPtr<FMassDebuggerArchetypeData>>& OutTransientArchetypesMap);
 
@@ -209,6 +220,8 @@ protected:
 	void OnEntitySelected(const FMassEntityManager& EntityManager, const FMassEntityHandle EntityHandle);
 
 public:
+	TWeakPtr<SMassDebugger> DebuggerWindow;
+
 	FOnRefresh OnRefreshDelegate;
 	FOnProcessorsSelected OnProcessorsSelectedDelegate;
 	FOnArchetypesSelected OnArchetypesSelectedDelegate;
@@ -226,13 +239,12 @@ public:
 		FName Label;
 		TArray<TSharedPtr<FMassDebuggerProcessorData>> Container;
 	};
+
 	TArray<TSharedPtr<FProcessorCollection>> CachedProcessorCollections;
 	TArray<TSharedPtr<FMassDebuggerProcessorData>> SelectedProcessors;
-
 	TArray<TSharedPtr<FMassDebuggerArchetypeData>> CachedAllArchetypes;
 	TArray<TSharedPtr<FMassDebuggerArchetypeData>> CachedArchetypeRepresentatives;
 	TArray<TSharedPtr<FMassDebuggerArchetypeData>> SelectedArchetypes;
-
 	TArray<TSharedPtr<FMassDebuggerProcessingGraph>> CachedProcessingGraphs;
 
 	TMap<FMassArchetypeHandle, TSharedPtr<FMassDebuggerArchetypeData>> HandleToArchetypeMap;
@@ -248,8 +260,8 @@ public:
 
 protected:	
 	TArray<TSharedPtr<FMassDebuggerProcessorData>> AllCachedProcessors;
+	TArray<TWeakPtr<SMassEntitiesView>> EntityViews;
 	FName SelectedFragmentName;
-
 public:
 	UE_DEPRECATED(5.6, "CachedProcessors property is now deprecated. Use CachedProcessorCollections instead. ")
 	TArray<TSharedPtr<FMassDebuggerProcessorData>> CachedProcessors;
diff --git a/Engine/Source/Editor/MassEntityDebugger/Private/SMassArchetype.cpp b/Engine/Source/Editor/MassEntityDebugger/Private/SMassArchetype.cpp
--- a/Engine/Source/Editor/MassEntityDebugger/Private/SMassArchetype.cpp
+++ b/Engine/Source/Editor/MassEntityDebugger/Private/SMassArchetype.cpp
@@ -4,6 +4,7 @@
 #include "MassDebuggerModel.h"
 #include "MassEntityTypes.h"
 #include "SMassBitSet.h"
+#include "SMassEntitiesList.h"
 #include "Widgets/Text/SRichTextBlock.h"
 #include "Styling/AppStyle.h"
 #include "Widgets/Input/SButton.h"
@@ -25,16 +26,6 @@ void SMassArchetype::Construct(const FArguments& InArgs, TSharedPtr<FMassDebugge
 	ArchetypeData = InArchetypeData;
 	DebuggerModel = InDebuggerModel;
 
-	// limit the number of entities shown to avoid perf issues
-	// @todo replace this temp fix with virtualized scrolling of the entire entity list (and filtering)
-	constexpr int MaxEntitiesToDisplay = 1000;
-	const int EntityDisplayCount = FMath::Min(ArchetypeData->EntityHandles.Num(), MaxEntitiesToDisplay);
-	EntityHandles.Reserve(EntityDisplayCount);
-	for (int i = 0; i < EntityDisplayCount; i++)
-	{
-		EntityHandles.Add(MakeShared<FMassEntityHandle>(ArchetypeData->EntityHandles[i]));
-	}
-
 	const FMassDebuggerArchetypeData* BaseArchetypeDebugData = InBaseArchetypeData.Get();
 	const FMassDebuggerArchetypeData& ArchetypeDebugData = *InArchetypeData.Get();
 
@@ -49,6 +40,16 @@ void SMassArchetype::Construct(const FArguments& InArgs, TSharedPtr<FMassDebugge
 		, InArchetypeData->LabelLong
 	};
 
+	Box->AddSlot()
+		.AutoHeight()
+		.Padding(0, 4)
+		[
+			SNew(SButton)
+				.Text(LOCTEXT("ShowEntities", "Show Entities"))
+				.ContentPadding(4)
+				.OnClicked(FOnClicked::CreateSP(this, &SMassArchetype::ShowEntities))
+		];
+
 	Box->AddSlot()
 		.AutoHeight()
 		.Padding(0, 4)
@@ -78,7 +79,7 @@ void SMassArchetype::Construct(const FArguments& InArgs, TSharedPtr<FMassDebugge
 		.Padding(0, 4)
 		[
 			SNew(STextBlock)
-				.Text(ArchetypeDescription)
+			.Text(ArchetypeDescription)
 		];
 
 	if (ArchetypeDebugData.ArchetypeStats.EntitiesCount != 0 && ArchetypeDebugData.ArchetypeStats.ChunksCount != 0)
@@ -122,112 +123,42 @@ void SMassArchetype::Construct(const FArguments& InArgs, TSharedPtr<FMassDebugge
 	TSharedRef<SVerticalBox> MainBox = SNew(SVerticalBox);
 
 	MainBox->AddSlot()
-		.AutoHeight()
-		.Padding(0, 4)
-		[
-			SNew(SButton)
-				.OnClicked_Lambda([this, Box]() -> FReply
-				{
-					bBitSetsVisible = !bBitSetsVisible;
-					Box->SetVisibility(bBitSetsVisible ? EVisibility::Visible : EVisibility::Collapsed);
-					return FReply::Handled();
-				})
-				[
-					SNew(SRichTextBlock)
-					.Text(FText::Join(FText::FromString(TEXT(": ")), LabelBits))
-					.DecoratorStyleSet(&FAppStyle::Get())
-					.TextStyle(FAppStyle::Get(), "LargeText")
-				]
-		];
-
-	MainBox->AddSlot()
-		.AutoHeight()
-		.Padding(0, 4)
-		[
-			Box
-		];
-
-	TSharedRef<SListView<TSharedPtr<FMassEntityHandle>>> ListViewWidget = SNew(SListView<TSharedPtr<FMassEntityHandle>>)
-		.ListItemsSource(&EntityHandles)
-		.OnGenerateRow(this, &SMassArchetype::OnGenerateRow)
-		.OnSelectionChanged(this, &SMassArchetype::OnEntitySelected);
-
-	MainBox->AddSlot()
-		.AutoHeight()
-		.Padding(0, 4)
+	.AutoHeight()
+	.Padding(0, 4)
+	[
+		SNew(SHorizontalBox)
+		+ SHorizontalBox::Slot()
 		[
-			SNew(SButton)
-				.Text(LOCTEXT("EntityList", "Entity List"))
-				.OnClicked_Lambda([this, ListViewWidget]() -> FReply
-				{
-					bEntitiesVisible = !bEntitiesVisible;
-					ListViewWidget->SetVisibility(bEntitiesVisible ? EVisibility::Visible : EVisibility::Collapsed);
-					return FReply::Handled();
-				})
-		];
-
+			SNew(SRichTextBlock)
+			.Text(FText::Join(FText::FromString(TEXT(": ")), LabelBits))
+			.DecoratorStyleSet(&FAppStyle::Get())
+			.TextStyle(FAppStyle::Get(), "LargeText")
+		]
+	];
 	MainBox->AddSlot()
-		.FillHeight(1.0f)
-		.Padding(0, 4)
-		[
-			ListViewWidget
-		];
+	.AutoHeight()
+	.Padding(0, 4)
+	[
+		Box
+	];
 
 	ChildSlot
+	[
+		SNew(SBorder)
+		.Padding(5.0f)
 		[
-			SNew(SBorder)
-				.Padding(5.0f)
-				[
-					MainBox
-				]
-		];
+			MainBox
+		]
+	];
 }
 
-void SMassArchetype::OnEntitySelected(TSharedPtr<FMassEntityHandle> SelectedItem, ESelectInfo::Type SelectInfo)
+FReply SMassArchetype::ShowEntities()
 {
-#if WITH_MASSENTITY_DEBUG
-	if (SelectedItem.IsValid() && DebuggerModel->Environment.IsValid())
+	if (DebuggerModel.IsValid() && ArchetypeData.IsValid())
 	{
-		FMassDebugger::SelectEntity(*(DebuggerModel->Environment->GetEntityManager()), *SelectedItem);
+		DebuggerModel->ShowEntitiesView(0, ArchetypeData->Handle);
 	}
-#endif
-}
-
-TSharedRef<ITableRow> SMassArchetype::OnGenerateRow(TSharedPtr<FMassEntityHandle> InItem, const TSharedRef<STableViewBase>& OwnerTable)
-{
-	return SNew(STableRow<TSharedPtr<FMassEntityHandle>>, OwnerTable)
-	.Content()
-	[
-#if WITH_MASSENTITY_DEBUG
-		SNew(SButton)
-		.Text(FText::FromString(InItem->DebugGetDescription()))
-		.OnClicked_Lambda([this, InItem]()
-		{
-			if (InItem.IsValid() && DebuggerModel->Environment.IsValid())
-			{
-				FMassDebugger::SelectEntity(*(DebuggerModel->Environment->GetEntityManager()), *InItem);
-			}
-			return FReply::Handled();
-		})
-		.OnHovered_Lambda([this, InItem]()
-		{
-			if (InItem.IsValid() && DebuggerModel->Environment.IsValid())
-			{
-				FMassDebugger::HighlightEntity(*(DebuggerModel->Environment->GetEntityManager()), *InItem);
-			}
-		})
-		.OnUnhovered_Lambda([this]()
-		{
-			if (DebuggerModel->Environment.IsValid())
-			{
-				FMassDebugger::HighlightEntity(*(DebuggerModel->Environment->GetEntityManager()), FMassEntityHandle());
-			}
-		})
-#else
-		SNew(STextBlock)
-			.Text(LOCTEXT("MassEntityDebuggingNotEnabled", "Mass Entity Debugging Not Enabled for this configuration"))
-#endif
-	];
+	return FReply::Handled();
 }
 
 #undef LOCTEXT_NAMESPACE
\ No newline at end of file
diff --git a/Engine/Source/Editor/MassEntityDebugger/Private/SMassArchetype.h b/Engine/Source/Editor/MassEntityDebugger/Private/SMassArchetype.h
--- a/Engine/Source/Editor/MassEntityDebugger/Private/SMassArchetype.h
+++ b/Engine/Source/Editor/MassEntityDebugger/Private/SMassArchetype.h
@@ -19,10 +19,8 @@ public:
 
 protected:
 	TSharedPtr<FMassDebuggerArchetypeData> ArchetypeData;
-	void OnEntitySelected(TSharedPtr<FMassEntityHandle> SelectedItem, ESelectInfo::Type SelectInfo);
-	TSharedRef<ITableRow> OnGenerateRow(TSharedPtr<FMassEntityHandle> InItem, const TSharedRef<STableViewBase>& OwnerTable);
 	TSharedPtr<FMassDebuggerModel> DebuggerModel;
-	TArray<TSharedPtr<FMassEntityHandle>> EntityHandles;
 	bool bBitSetsVisible = true;
 	bool bEntitiesVisible = true;
+	FReply ShowEntities();
 };
diff --git a/Engine/Source/Editor/MassEntityDebugger/Private/SMassDebugger.cpp b/Engine/Source/Editor/MassEntityDebugger/Private/SMassDebugger.cpp
--- a/Engine/Source/Editor/MassEntityDebugger/Private/SMassDebugger.cpp
+++ b/Engine/Source/Editor/MassEntityDebugger/Private/SMassDebugger.cpp
@@ -5,6 +5,7 @@
 #include "SMassProcessorsView.h"
 #include "SMassProcessingView.h"
 #include "SMassArchetypesView.h"
+#include "SMassEntitiesView.h"
 #include "MassDebuggerModel.h"
 #include "Engine/Engine.h"
 #include "CoreGlobals.h"
@@ -36,6 +37,7 @@ namespace UE::Mass::Debugger::Private
 	const FName ProcessorsTabId("Processors");
 	const FName ProcessingGraphTabId("Processing Graphs");
 	const FName ArchetypesTabId("Archetypes");
+	const FName EntitiesTabId("Entities");
 
 	bool IsSupportedWorldType(const EWorldType::Type WorldType)
 	{
@@ -101,6 +103,9 @@ void SMassDebugger::Construct(const FArguments& InArgs, const TSharedRef<SDockTa
 		.SetDisplayName(LOCTEXT("BreakpointsTabTitle", "Breakpoints"))
 		.SetGroup(AppMenuGroup);
 
+	TabManager->RegisterTabSpawner(EntitiesTabId, FOnSpawnTab::CreateRaw(this, &SMassDebugger::SpawnEntitiesTab))
+		.SetDisplayName(LOCTEXT("EntitiesTabTitle", "Entities"))
+		.SetGroup(AppMenuGroup);
 
 	// Default Layout
 	TSharedRef<FTabManager::FLayout> Layout = CreateDefaultLayout();
@@ -123,6 +128,8 @@ void SMassDebugger::Construct(const FArguments& InArgs, const TSharedRef<SDockTa
 			}
 		)
 	);
+
+	DebuggerModel->DebuggerWindow = SharedThis(this).ToWeakPtr();
 }
 
 TSharedRef<FTabManager::FLayout> SMassDebugger::CreateDefaultLayout()
@@ -158,13 +165,15 @@ TSharedRef<FTabManager::FLayout> SMassDebugger::CreateDefaultLayout()
 					(
 						FTabManager::NewStack()
 						->AddTab(BreakpointsTabId, ETabState::OpenedTab)
-						->SetForegroundTab(ProcessorsTabId)
+						->SetForegroundTab(BreakpointsTabId)
 					)
 				)
 				->Split
 				(
 					FTabManager::NewStack()
 					->AddTab(ArchetypesTabId, ETabState::OpenedTab)
+					->AddTab(EntitiesTabId, ETabState::OpenedTab)
+					->SetForegroundTab(ArchetypesTabId)
 				)
 			)
 		);
@@ -236,6 +245,7 @@ TSharedRef<SWidget> SMassDebugger::GenerateWindowMenu()
 	FUIAction ArchetypesUIAction(FExecuteAction::CreateSP(this, &SMassDebugger::ShowArchetypesView));
 	FUIAction ProcessingGraphsUIAction(FExecuteAction::CreateSP(this, &SMassDebugger::ShowProcessingGraphsView));
 	FUIAction BreakpointsUIAction(FExecuteAction::CreateSP(this, &SMassDebugger::ShowBreakpointsView));
+	FUIAction EntitesUIAction(FExecuteAction::CreateSP(this, &SMassDebugger::ShowEntitesView));
 	FUIAction ResetLayoutUIAction(FExecuteAction::CreateSP(this, &SMassDebugger::ResetLayout));
 
 	MenuBuilder.BeginSection(TEXT("Tabs"));
@@ -243,6 +253,7 @@ TSharedRef<SWidget> SMassDebugger::GenerateWindowMenu()
 	MenuBuilder.AddMenuEntry(LOCTEXT("ArchetypesTabLabel", "Archetypes"), LOCTEXT("ArchtypesTabTooltip", "Show Archtypes Tab"), FSlateIcon(), ArchetypesUIAction);
 	MenuBuilder.AddMenuEntry(LOCTEXT("ProcessingGraphsTabLabel", "Processing Graphs"), LOCTEXT("ProcessingGraphsTabTooltip", "Show Processing Graphs Tab"), FSlateIcon(), ProcessingGraphsUIAction);
 	MenuBuilder.AddMenuEntry(LOCTEXT("BreakpointsTabLabel", "Breakpoints"), LOCTEXT("BreakpointsTabTooltip", "Show Breakpoints Tab"), FSlateIcon(), BreakpointsUIAction);
+	MenuBuilder.AddMenuEntry(LOCTEXT("EntitiesTabLabel", "Entities"), LOCTEXT("EntitiesTabTooltip", "Show Entities Tab"), FSlateIcon(), EntitesUIAction);
 	MenuBuilder.AddSeparator();
 	MenuBuilder.AddMenuEntry(LOCTEXT("ResetLayoutLabel", "Reset Layout"), LOCTEXT("ResetLayoutTooltip", "Reset the Mass Debugger Layout"), FSlateIcon(), ResetLayoutUIAction);
 
@@ -385,6 +396,17 @@ TSharedRef<SDockTab> SMassDebugger::SpawnBreakpointsTab(const FSpawnTabArgs& Arg
 	return MajorTab;
 }
 
+TSharedRef<SDockTab> SMassDebugger::SpawnEntitiesTab(const FSpawnTabArgs& Args)
+{
+	const TSharedRef<SDockTab> MajorTab = SNew(SDockTab)
+		.TabRole(ETabRole::PanelTab);
+
+	TSharedPtr<SWidget> TabContent = SNew(SMassEntitiesView, DebuggerModel, 0);
+	MajorTab->SetContent(TabContent.ToSharedRef());
+
+	return MajorTab;
+}
+
 bool SMassDebugger::CanRefreshData()
 {
 	return DebuggerModel->HasEnvironmentSelected();
@@ -415,6 +437,11 @@ void SMassDebugger::ShowProcessingGraphsView() const
 	TabManager->TryInvokeTab(UE::Mass::Debugger::Private::ProcessingGraphTabId);
 }
 
+void SMassDebugger::ShowEntitesView() const
+{
+	TabManager->TryInvokeTab(UE::Mass::Debugger::Private::EntitiesTabId);
+}
+
 void SMassDebugger::ResetLayout()
 {
 	TSharedPtr<SWindow> Window = FSlateApplication::Get().FindWidgetWindow(AsShared());
diff --git a/Engine/Source/Editor/MassEntityDebugger/Private/SMassDebugger.h b/Engine/Source/Editor/MassEntityDebugger/Private/SMassDebugger.h
--- a/Engine/Source/Editor/MassEntityDebugger/Private/SMassDebugger.h
+++ b/Engine/Source/Editor/MassEntityDebugger/Private/SMassDebugger.h
@@ -56,12 +56,21 @@ public:
 
 	virtual bool SupportsKeyboardFocus() const override { return true; }
 
+	void ShowProcessorView() const;
+	void ShowArchetypesView() const;
+	void ShowBreakpointsView() const;
+	void ShowProcessingGraphsView() const;
+	void ShowEntitesView() const;
+	void ResetLayout();
+	void ShowSelectedView() const;
+
 protected:
 	TSharedRef<SDockTab> SpawnToolbar(const FSpawnTabArgs& Args);
 	TSharedRef<SDockTab> SpawnProcessorsTab(const FSpawnTabArgs& Args);
 	TSharedRef<SDockTab> SpawnProcessingTab(const FSpawnTabArgs& Args);
 	TSharedRef<SDockTab> SpawnArchetypesTab(const FSpawnTabArgs& Args);
 	TSharedRef<SDockTab> SpawnBreakpointsTab(const FSpawnTabArgs& Args);
+	TSharedRef<SDockTab> SpawnEntitiesTab(const FSpawnTabArgs& Args);
 
 	TSharedRef<FTabManager::FLayout> CreateDefaultLayout();
 	TSharedRef<SWidget> GenerateWindowMenu();
@@ -75,13 +84,6 @@ protected:
 	bool CanRefreshData();
 	void RefreshData();
 
-	void ShowProcessorView() const;
-	void ShowArchetypesView() const;
-	void ShowBreakpointsView() const;
-	void ShowProcessingGraphsView() const;
-	void ResetLayout();
-	void ShowSelectedView() const;
-
 #if WITH_MASSENTITY_DEBUG
 void OnProcessorProviderRegistered(const FMassDebugger::FEnvironment& Environment);
 #endif // WITH_MASSENTITY_DEBUG
diff --git a/Engine/Source/Editor/MassEntityDebugger/Private/SMassEntitiesList.cpp b/Engine/Source/Editor/MassEntityDebugger/Private/SMassEntitiesList.cpp
new file mode 100644
--- /dev/null
+++ b/Engine/Source/Editor/MassEntityDebugger/Private/SMassEntitiesList.cpp
@@ -0,0 +1,655 @@
+// Copyright Epic Games, Inc. All Rights Reserved.
+
+#include "SMassEntitiesList.h"
+#include "MassDebuggerModel.h"
+#include "MassEntityTypes.h"
+#include "SMassBitSet.h"
+#include "Widgets/Text/SRichTextBlock.h"
+#include "Styling/AppStyle.h"
+#include "Widgets/Input/SButton.h"
+#include "Widgets/Views/SListView.h"
+#include "MassArchetypeData.h"
+#include "PropertyEditorModule.h"
+#include "Modules/ModuleManager.h"
+#include "UObject/StructOnScope.h"
+#include "IStructureDetailsView.h"
+#include "ISinglePropertyView.h"
+#include "IPropertyRowGenerator.h"
+#include "IDetailTreeNode.h"
+#include "DetailTreeNode.h"
+#include "Widgets/Layout/SGridPanel.h"
+#include "Widgets/Layout/SSeparator.h"
+#include "Widgets/Layout/SScrollBar.h"
+#include "Widgets/Layout/SScrollBox.h"
+#include "Widgets/Input/SCheckBox.h"
+#include "Widgets/Input/SComboButton.h"
+
+#define LOCTEXT_NAMESPACE "SMassDebugger"
+
+//----------------------------------------------------------------------//
+// SMassEntitiesList
+//----------------------------------------------------------------------//
+void SMassEntitiesList::Construct(const FArguments& InArgs, TSharedRef<FMassDebuggerModel> InDebuggerModel)
+{
+#if WITH_MASSENTITY_DEBUG
+	EntityGrid = SNew(SGridPanel);
+	FragmentSelectBox = SNew(SBox);
+
+	DebuggerModel = InDebuggerModel;
+
+	EntityScroll = SNew(SScrollBar)
+		.AlwaysShowScrollbar(true)
+		.Orientation(EOrientation::Orient_Vertical);
+	EntityScroll->SetOnUserScrolled(FOnUserScrolled::CreateSP(this, &SMassEntitiesList::OnUserScrolled));
+	
+	ChildSlot
+	[
+		SNew(SBorder)
+		.Padding(5.0f)
+		[
+			SNew(SVerticalBox)
+			+ SVerticalBox::Slot()
+			.AutoHeight()
+			.HAlign(HAlign_Left)
+			[
+				FragmentSelectBox.ToSharedRef()
+			]
+			+ SVerticalBox::Slot()
+			.FillHeight(1.f)
+			[
+				SNew(SHorizontalBox)
+				+ SHorizontalBox::Slot()
+				.FillWidth(1.f)
+				[
+					SNew(SScrollBox)
+					.Orientation(Orient_Horizontal)
+					+SScrollBox::Slot()
+					[
+						EntityGrid.ToSharedRef()
+					]
+				]
+				+ SHorizontalBox::Slot()
+				.AutoWidth()
+				.Padding(10.0f)
+				[
+					EntityScroll.ToSharedRef()
+				]
+			]
+		]
+	];
+
+	SetEntities(InArgs._Entities);
+#else
+	ChildSlot
+	[
+		SNew(STextBlock)
+			.Text(LOCTEXT("MassEntityDebuggingNotEnabled", "Mass Entity Debugging Not Enabled for this configuration"))
+	];
+#endif
+}
+
+void SMassEntitiesList::OnUserScrolled(float InScrollOffsetFraction)
+{
+	float Top = float(GridRows.Num()) * (InScrollOffsetFraction);
+	int NewTop = (int)Top;
+	SetTopIndex(NewTop);
+}
+
+void SMassEntitiesList::BuildGrid()
+{
+#if WITH_MASSENTITY_DEBUG
+	EntityGrid->ClearChildren();
+
+	PopulateGridColumns();
+	PrepareRowData();
+	PopulateEntityData();
+#endif
+}
+
+void SMassEntitiesList::AddPropertyRecursive(TSharedPtr<SHorizontalBox> HBox, TSharedPtr<SVerticalBox> VBox, TSharedPtr<IPropertyHandle> Prop, bool bShowName)
+{
+#if WITH_MASSENTITY_DEBUG
+		uint32 NumChildren = 0;
+		Prop->GetNumChildren(NumChildren);
+		if (NumChildren > 0)
+		{
+			VBox = SNew(SVerticalBox);
+			HBox->AddSlot()
+			.AutoWidth()
+			.HAlign(EHorizontalAlignment::HAlign_Center)
+			.VAlign(EVerticalAlignment::VAlign_Top)
+			[
+				VBox.ToSharedRef()
+			];
+			if (bShowName)
+			{
+				VBox->AddSlot()
+				.AutoHeight()
+				.HAlign(EHorizontalAlignment::HAlign_Center)
+				.VAlign(EVerticalAlignment::VAlign_Top)
+				[
+					Prop->CreatePropertyNameWidget()
+				];
+			}
+
+			for (uint32 i = 0; i < NumChildren; i++)
+			{
+				AddPropertyRecursive(HBox, VBox, Prop->GetChildHandle(i), true);
+			}
+		}
+		else
+		{
+			if (bShowName)
+			{
+				VBox->AddSlot()
+				.AutoHeight()
+				.HAlign(EHorizontalAlignment::HAlign_Left)
+				.VAlign(EVerticalAlignment::VAlign_Top)
+				[
+					SNew(SHorizontalBox)
+					+ SHorizontalBox::Slot()
+					.AutoWidth()
+					.HAlign(EHorizontalAlignment::HAlign_Left)
+					.VAlign(EVerticalAlignment::VAlign_Top)
+					[
+						Prop->CreatePropertyNameWidget()
+					]
+					+ SHorizontalBox::Slot()
+					.AutoWidth()
+					.HAlign(EHorizontalAlignment::HAlign_Left)
+					.VAlign(EVerticalAlignment::VAlign_Top)
+					[
+						Prop->CreatePropertyValueWidget()
+					]
+				];
+			}
+			else
+			{
+				VBox->AddSlot()
+				.HAlign(EHorizontalAlignment::HAlign_Left)
+				.VAlign(EVerticalAlignment::VAlign_Top)
+				[
+					Prop->CreatePropertyValueWidget()
+				];
+			}
+		}
+#endif
+}
+
+void SMassEntitiesList::PopulateEntityData()
+{
+#if WITH_MASSENTITY_DEBUG
+	for (int i = 0; i < NumGridRowsDrawn; i++)
+	{
+		FillGridRow(i);
+	}
+#endif
+}
+
+void SMassEntitiesList::SetEntities(const TArray<FMassEntityHandle>& InEntities)
+{
+#if WITH_MASSENTITY_DEBUG
+	EntityScroll->SetState(0.0f, (float)NumGridRowsDrawn / FMath::Max(1.f, (float)InEntities.Num()));
+	TopIndex = 0;
+	AvailableFragmentNames.Reset();
+	// ensure that the selected fragments are preserved even if they're not present on the new set of entites:
+	AvailableFragmentNames.Append(SelectedFragmentNames);
+	
+	GridRows.Reset();
+	GridRows.SetNum(InEntities.Num(), EAllowShrinking::No);
+
+	if (DebuggerModel.IsValid() && DebuggerModel->Environment.IsValid() && DebuggerModel->Environment->EntityManager.IsValid())
+	{
+		const FMassEntityManager& EntityManager = *(DebuggerModel->Environment->EntityManager.Pin());
+		TSet<FMassArchetypeHandle> SearchedArchetypes;
+
+		for (int i = 0; i < InEntities.Num(); i++)
+		{
+			const FMassEntityHandle& Entity = InEntities[i];
+			GridRows[i].Entity = Entity;
+
+			FMassArchetypeHandle ArchetypeHandle = EntityManager.GetArchetypeForEntity(Entity);
+			
+			if (ArchetypeHandle.IsValid() && !SearchedArchetypes.Contains(ArchetypeHandle))
+			{
+				SearchedArchetypes.Add(ArchetypeHandle);
+				EntityManager.ForEachArchetypeFragmentType(ArchetypeHandle, [this](const UScriptStruct* FragmentType)
+				{
+					AvailableFragmentNames.AddUnique(FragmentType->GetFName());
+				});
+			}
+
+		}
+	}
+
+	AvailableFragmentNames.Sort([](const FName& A, const FName& B) { return A.Compare(B) < 0; });
+
+	CreateFragmentSelectDropdown();
+	PrepareRowData();
+	BuildGrid();
+#endif
+}
+
+void SMassEntitiesList::PrepareRowData()
+{
+#if WITH_MASSENTITY_DEBUG
+	if (!DebuggerModel.IsValid() || !DebuggerModel->Environment.IsValid() || !DebuggerModel->Environment->GetEntityManager().IsValid())
+	{
+		return;
+	}
+
+	FPropertyEditorModule& PropertyEditorModule = FModuleManager::LoadModuleChecked<FPropertyEditorModule>("PropertyEditor");
+	const FMassEntityManager& EntityManager = *(DebuggerModel->Environment->GetEntityManager());
+
+	FPropertyRowGeneratorArgs GeneratorArgs;
+	GeneratorArgs.bShouldShowHiddenProperties = true;
+
+	int FragmentDisplayCount = SelectedFragmentTypes.Num();
+
+	int MaxRow = FMath::Min(TopIndex + NumGridRowsDrawn, GridRows.Num());
+	for (int RowIndex = TopIndex; RowIndex < MaxRow; RowIndex++)
+	{
+		FGridRow& Row = GridRows[RowIndex];
+		if (Row.bDirty)
+		{
+			Row.bDirty = false;
+			Row.FragmentInfo.SetNum(FragmentDisplayCount);
+
+			for (int i = 0; i < FragmentDisplayCount; i++)
+			{
+				if (Row.FragmentInfo[i].StructType != SelectedFragmentTypes[i])
+				{
+					Row.FragmentInfo[i].StructType = SelectedFragmentTypes[i];
+					Row.FragmentInfo[i].StructData = FMassDebugger::GetFragmentData(EntityManager, SelectedFragmentTypes[i], Row.Entity);
+
+					if (Row.FragmentInfo[i].StructData.IsValid())
+					{
+						Row.FragmentInfo[i].PropertyRowGenerator = PropertyEditorModule.CreatePropertyRowGenerator(GeneratorArgs);
+						Row.FragmentInfo[i].PropertyRowGenerator->SetStructure(Row.FragmentInfo[i].StructData);
+					}
+				}
+			}
+		}
+	}
+#endif
+}
+
+TSharedPtr<IPropertyHandle> FindPropertyHandle(const FProperty* Property, TArray<TSharedRef<IDetailTreeNode>>& NodesToSearch)
+{
+#if WITH_MASSENTITY_DEBUG
+	TArray<TSharedRef<IDetailTreeNode>> Children;
+	while (NodesToSearch.Num() > 0)
+	{
+		TSharedRef<IDetailTreeNode> CurNode = NodesToSearch.Pop(EAllowShrinking::No);
+		if (CurNode->GetNodeType() == EDetailNodeType::Item)
+		{
+			TSharedPtr<IPropertyHandle> PropertyHandle = CurNode->CreatePropertyHandle();
+			if (PropertyHandle.IsValid() && PropertyHandle->GetProperty() == Property)
+			{
+				return PropertyHandle;
+			}
+		}
+
+		CurNode->GetChildren(Children, true);
+		NodesToSearch.Append(Children);
+	}
+#endif
+	return nullptr;
+}
+
+SMassEntitiesList::FFragmentColumn* SMassEntitiesList::GetFragmentColumn(const UScriptStruct* StructType)
+{
+#if WITH_MASSENTITY_DEBUG
+	for (FFragmentColumn& Column : GridColumns)
+	{
+		if (Column.StructType == StructType)
+		{
+			return &Column;
+		}
+	}
+#endif
+	return nullptr;
+}
+
+void SMassEntitiesList::FillGridRow(int RowIndex)
+{
+#if WITH_MASSENTITY_DEBUG
+	int32 GridRow = RowIndex * 2 + NumHeaderRows;
+	RowIndex += TopIndex;
+	if (RowIndex >= GridRows.Num() || RowIndex < 0)
+	{
+		return;
+	}
+	FGridRow& Row = GridRows[RowIndex];
+	FMassEntityHandle Entity = Row.Entity;
+	TWeakPtr<FMassDebuggerModel> WeakModel = DebuggerModel.ToWeakPtr();
+	TWeakPtr<const FMassEntityManager> WeakEntityManager = DebuggerModel->Environment->GetEntityManager().ToWeakPtr();
+
+	// Add the Entity Handle to the grid
+	EntityGrid->AddSlot(0, GridRow)
+	[
+		SNew(SButton)
+		.VAlign(VAlign_Fill)
+		.HAlign(HAlign_Fill)
+		.Text(FText::FromString(Entity.DebugGetDescription()))
+		.OnClicked_Lambda([WeakEntityManager, Entity]()
+		{
+			if (Entity.IsValid() && WeakEntityManager.IsValid())
+			{
+				FMassDebugger::SelectEntity(*WeakEntityManager.Pin(), Entity);
+			}
+			return FReply::Handled();
+		})
+	];
+
+	TArray<TSharedRef<IDetailTreeNode>> NodesToSearch;
+	NodesToSearch.Reserve(100);
+
+	//for every selected fragment type
+	for (FGridRow::FFragmentInfo& FragmentInf : Row.FragmentInfo)
+	{
+		if (FragmentInf.PropertyRowGenerator.IsValid())
+		{
+			const UScriptStruct* FragmentType = FragmentInf.StructType;
+
+			// the entity for this row has this fragment
+			SMassEntitiesList::FFragmentColumn* Column = GetFragmentColumn(FragmentInf.StructType);
+			if (Column)
+			{
+				EntityGrid->AddSlot(Column->ColumnNum, GridRow)
+				.HAlign(EHorizontalAlignment::HAlign_Center)
+				.VAlign(EVerticalAlignment::VAlign_Top)
+				[
+					SNew(SVerticalBox)
+					+ SVerticalBox::Slot()
+					.AutoHeight()
+					[
+						SNew(SHorizontalBox)
+						+ SHorizontalBox::Slot()
+						.AutoWidth()
+						.Padding(4.0f)
+						[
+							SNew(SButton)
+							.ButtonStyle(FAppStyle::Get(), "FlatButton")
+							.ContentPadding(4.0f)
+							.OnClicked_Lambda([Entity, FragmentType, WeakModel]()
+							{
+								TSharedPtr<FMassDebuggerModel> Model = WeakModel.Pin();
+								if (Model.IsValid() && Model->Environment.IsValid() && Model->Environment->EntityManager.IsValid())
+								{
+									FMassDebugger::SetFragmentWriteBreak(*Model->Environment->EntityManager.Pin().Get(), FragmentType, Entity);
+								}
+								return FReply::Handled();
+							})
+							[
+								SNew(SImage)
+								.Image(FAppStyle::Get().GetBrush("GenericStop"))
+								.DesiredSizeOverride(FVector2D(16, 16))
+							]
+						]
+					]
+				];
+
+				NodesToSearch = FragmentInf.PropertyRowGenerator->GetRootTreeNodes();
+
+				// loop through the properties on this fragment and fill in those columns
+				for (FFragmentColumn::FPropertyColumn& PropColumn : Column->PropertyColumns)
+				{
+					// some of the properties we want to find might be nested in categories so we need to search the tree
+					TSharedPtr<IPropertyHandle> PropertyHandle = FindPropertyHandle(PropColumn.Property, NodesToSearch);
+					if (PropertyHandle.IsValid())
+					{
+						// Found a top-level property on the fragment
+
+						TSharedPtr<SHorizontalBox> HBox = SNew(SHorizontalBox);
+						TSharedPtr<SVerticalBox> VBox = SNew(SVerticalBox);
+						HBox->AddSlot()
+						.HAlign(EHorizontalAlignment::HAlign_Center)
+						.VAlign(EVerticalAlignment::VAlign_Top)
+						.AutoWidth()
+						[
+							VBox.ToSharedRef()
+						];
+
+						AddPropertyRecursive(HBox, VBox, PropertyHandle);
+
+						EntityGrid->AddSlot(PropColumn.ColumnNum, GridRow)
+						[
+							HBox.ToSharedRef()
+						];
+					}
+				}
+
+			}
+		}
+	}
+#endif
+}
+
+void SMassEntitiesList::PopulateGridColumns()
+{
+#if WITH_MASSENTITY_DEBUG
+	GridColumns.Reset();
+
+	// Column 0 reserved for entity handle
+	EntityGrid->AddSlot(0,1)
+	[
+		SNew(STextBlock)
+		.Text(LOCTEXT("EntityHandle", "Entity Handle"))
+	];
+
+	int CurColumn = 1;
+	for (FName FragmentName : SelectedFragmentNames)
+	{
+		const UScriptStruct* FragmentStructType = FMassDebugger::GetFragmentTypeFromName(FragmentName);
+		if (!FragmentStructType)
+		{
+			continue;
+		}
+
+		FFragmentColumn& FragmentColumn = GridColumns.AddDefaulted_GetRef();
+		FragmentColumn.ColumnNum = CurColumn++;
+		FragmentColumn.StructType = FragmentStructType;
+
+		for (TFieldIterator<FProperty> It(FragmentStructType); It; ++It)
+		{
+			FProperty* Property = *It;
+
+			FFragmentColumn::FPropertyColumn& Column = FragmentColumn.PropertyColumns.AddDefaulted_GetRef();
+			Column.ColumnNum = CurColumn;
+			Column.Property = Property;
+
+			EntityGrid->AddSlot(CurColumn, 1)
+			.HAlign(EHorizontalAlignment::HAlign_Left)
+			.VAlign(EVerticalAlignment::VAlign_Center)
+			.Padding(6.0f)
+			[
+				SNew(STextBlock)
+				.Text(FText::FromString(Property->GetName()))
+			];
+			++CurColumn;
+		}
+
+		SGridPanel::FSlot::FSlotArguments& Slot = EntityGrid->AddSlot(FragmentColumn.ColumnNum, 0)
+		.ColumnSpan(CurColumn - FragmentColumn.ColumnNum + 1)
+		.HAlign(EHorizontalAlignment::HAlign_Center)
+		.VAlign(EVerticalAlignment::VAlign_Center)
+		.Padding(6.0f)
+		[
+			SNew(STextBlock)
+			.Text(FText::FromName(FragmentName))
+		];
+
+		// Separator
+		EntityGrid->AddSlot(++CurColumn, 1)
+		.RowSpan(100)
+		[
+			SNew(SSeparator)
+			.Orientation(Orient_Vertical)
+		];
+		++CurColumn;
+	}
+#endif
+}
+
+void SMassEntitiesList::SetTopIndex(int32 NewTopIndex)
+{
+#if WITH_MASSENTITY_DEBUG
+	NewTopIndex = FMath::Clamp(NewTopIndex, 0, FMath::Max(GridRows.Num() - NumGridRowsDrawn, 0));
+	int32 Delta = NewTopIndex - TopIndex;
+	if (Delta == 0)
+	{
+		return;
+	}
+	const int32 MaxGridRow = NumHeaderRows + (2 * NumGridRowsDrawn);
+
+	EntityScroll->SetState((float)NewTopIndex / (float)GridRows.Num(), (float)NumGridRowsDrawn / FMath::Max(1.f, (float)GridRows.Num()));;
+
+	if (TPanelChildren<SGridPanel::FSlot>* ChildrenPtr = static_cast<TPanelChildren<SGridPanel::FSlot>*>(EntityGrid->GetChildren()))
+	{
+		TPanelChildren<SGridPanel::FSlot>& Children = *ChildrenPtr;
+
+		for (int32 SlotIndex = Children.Num() - 1; SlotIndex >= 0; --SlotIndex)
+		{
+			SGridPanel::FSlot& Slot = Children[SlotIndex];
+			int32 CurRow = Slot.GetRow();
+			int32 NewRow = CurRow - (2 * Delta);
+			if (CurRow >= NumHeaderRows)
+			{
+				if (NewRow >= NumHeaderRows && NewRow < MaxGridRow)
+				{
+					Slot.SetRow(NewRow);
+				}
+				else
+				{
+					Children.RemoveAt(SlotIndex);
+				}
+			}
+		}
+	}
+
+	TopIndex = NewTopIndex;
+	int FillFrom;
+	int FillTo;
+	if (Delta > 0)
+	{
+		FillFrom = FMath::Max(NumGridRowsDrawn - Delta, 0);
+		FillTo = NumGridRowsDrawn;
+	}
+	else
+	{
+		FillFrom = 0;
+		FillTo = FMath::Min(NumGridRowsDrawn, -Delta);
+	}
+	for (int i = FillFrom; i < FillTo; i++)
+	{
+		FillGridRow(i);
+	}
+	
+#endif
+}
+
+FReply SMassEntitiesList::OnClearAllSelectedFragmentsClicked()
+{
+	SelectedFragmentTypes.Reset();
+	SelectedFragmentNames.Reset();
+	BuildGrid();
+	return FReply::Handled();
+}
+
+void SMassEntitiesList::OnFragmentCheckStateChanged(ECheckBoxState NewState, FName FragmentName)
+{
+#if WITH_MASSENTITY_DEBUG
+	const bool bIsChecked = (NewState == ECheckBoxState::Checked);
+	if (bIsChecked)
+	{
+		SelectedFragmentNames.AddUnique(FragmentName);
+		SelectedFragmentNames.Sort([](const FName& A, const FName& B) { return A.Compare(B) < 0; });
+	}
+	else
+	{
+		SelectedFragmentNames.Remove(FragmentName);
+	}
+
+	SelectedFragmentTypes.Reset(SelectedFragmentNames.Num());
+	for (FName SelectedFragmentName : SelectedFragmentNames)
+	{
+		SelectedFragmentTypes.Add(FMassDebugger::GetFragmentTypeFromName(SelectedFragmentName));
+	}
+	
+	for (FGridRow& Row : GridRows)
+	{
+		Row.bDirty = true;
+	}
+
+	BuildGrid();
+#endif
+}
+
+ECheckBoxState SMassEntitiesList::GetFragmentCheckState(FName FragmentName) const
+{
+	if (SelectedFragmentNames.Contains(FragmentName))
+	{
+		return ECheckBoxState::Checked;
+	}
+	return ECheckBoxState::Unchecked;
+}
+
+void SMassEntitiesList::CreateFragmentSelectDropdown()
+{
+#if WITH_MASSENTITY_DEBUG
+	TSharedPtr<SVerticalBox> DropdownContent = SNew(SVerticalBox);
+
+	DropdownContent->AddSlot()
+	.AutoHeight()
+	.Padding(5.0f)
+	[
+		SNew(SButton)
+			.Text(LOCTEXT("ClearAll", "Clear All"))
+			.OnClicked(FOnClicked::CreateSP(this, &SMassEntitiesList::OnClearAllSelectedFragmentsClicked))
+	];
+
+	for (const FName& FragmentName : AvailableFragmentNames)
+	{
+		DropdownContent->AddSlot()
+		.AutoHeight()
+		.Padding(5.0f)
+		[
+			SNew(SCheckBox)
+			.OnCheckStateChanged(FOnCheckStateChanged::CreateSP(this, &SMassEntitiesList::OnFragmentCheckStateChanged, FragmentName))
+			.IsChecked(TAttribute<ECheckBoxState>::CreateSP(this, &SMassEntitiesList::GetFragmentCheckState, FragmentName))
+			.Content()
+			[
+				SNew(STextBlock).Text(FText::FromName(FragmentName))
+			]
+		];
+	}
+
+	TSharedRef<SComboButton> FragmentsButton = SNew(SComboButton)
+	.ButtonContent()
+	[
+		SNew(STextBlock)
+		.Text(FText::FromString(TEXT("Select Fragments")))
+	]
+	.MenuContent()
+	[
+		SNew(SBox)
+		[
+			DropdownContent.ToSharedRef()
+		]
+	]
+	.ComboButtonStyle(&FCoreStyle::Get().GetWidgetStyle<FComboButtonStyle>("ComboButton"))
+	.ButtonStyle(&FCoreStyle::Get().GetWidgetStyle<FButtonStyle>("Button"))
+	.ForegroundColor(FCoreStyle::Get().GetSlateColor("InvertedForeground"))
+	.ContentPadding(FMargin(5.0f));
+
+	FragmentSelectBox->SetContent(FragmentsButton);
+#endif
+}
+
+void SMassEntitiesList::RefreshEntityData()
+{
+	BuildGrid();
+}
+
+#undef LOCTEXT_NAMESPACE
\ No newline at end of file
diff --git a/Engine/Source/Editor/MassEntityDebugger/Private/SMassEntitiesList.h b/Engine/Source/Editor/MassEntityDebugger/Private/SMassEntitiesList.h
new file mode 100644
--- /dev/null
+++ b/Engine/Source/Editor/MassEntityDebugger/Private/SMassEntitiesList.h
@@ -0,0 +1,90 @@
+// Copyright Epic Games, Inc. All Rights Reserved.
+#pragma once
+
+#include "Widgets/SCompoundWidget.h"
+#include "MassDebugger.h"
+#include "MassDebuggerModel.h"
+
+class IPropertyHandle;
+class IPropertyRowGenerator;
+class FStructOnScope;
+class SGridPanel;
+class SBox;
+class SScrollBar;
+class SHorizontalBox;
+class SVerticalBox;
+
+class SMassEntitiesList : public SCompoundWidget
+{
+public:
+
+	SLATE_BEGIN_ARGS(SMassEntitiesList){}
+		SLATE_ARGUMENT(TArray<FMassEntityHandle>, Entities)
+	SLATE_END_ARGS()
+
+	void Construct(const FArguments& InArgs, TSharedRef<FMassDebuggerModel> InDebuggerModel);
+	void SetEntities(const TArray<FMassEntityHandle>& InEntities);
+	void RefreshEntityData();
+protected:
+
+	struct FFragmentColumn
+	{
+		int32 ColumnNum;
+		const UScriptStruct* StructType = nullptr;
+		struct FPropertyColumn
+		{
+			int32 ColumnNum;
+			const FProperty* Property = nullptr;
+		};
+		TArray<FPropertyColumn> PropertyColumns;
+	};
+
+	struct FGridRow
+	{
+		struct FFragmentInfo
+		{
+			const UScriptStruct* StructType = nullptr;
+			TSharedPtr<IPropertyRowGenerator> PropertyRowGenerator;
+
+			// this is a snapshot copy of the actual mass fragment data
+			TSharedPtr<FStructOnScope> StructData;
+		};
+		
+		TArray<FFragmentInfo> FragmentInfo;
+		FMassEntityHandle Entity;
+		bool bDirty = true;
+	};
+
+	static const int32 NumHeaderRows = 2;
+	static const int32 NumGridRowsDrawn = 32;
+	TArray<FGridRow> GridRows;
+	TArray<FFragmentColumn> GridColumns;
+
+	int32 TopIndex;
+	void BuildGrid();
+	void PopulateEntityData();
+	void PopulateGridColumns();
+	void PopulateRow(const FMassEntityManager& EntityManager, int RowIndex);
+	void PrepareRowData();
+	void FillGridRow(int Row);
+	FFragmentColumn* GetFragmentColumn(const UScriptStruct* StructType);
+	void SetTopIndex(int32 NewTopIndex);
+
+	void AddPropertyRecursive(TSharedPtr<SHorizontalBox> HBox, TSharedPtr<SVerticalBox> VBox, TSharedPtr<IPropertyHandle> Prop, bool bShowName = false);
+	TArray<FName> AvailableFragmentNames;
+	TArray<FName> SelectedFragmentNames;
+	TArray<const UScriptStruct*> SelectedFragmentTypes;
+	TSharedPtr<SWidget> FragmentSelectMenu;
+	FReply OnClearAllSelectedFragmentsClicked();
+
+	TSharedPtr<SGridPanel> EntityGrid;
+	TSharedPtr<SBox> FragmentSelectBox;
+
+	void OnFragmentCheckStateChanged(ECheckBoxState NewState, FName FragmentName);
+	ECheckBoxState GetFragmentCheckState(FName FragmentName) const;
+	void CreateFragmentSelectDropdown();
+	TSharedPtr<FMassDebuggerModel> DebuggerModel;
+	
+	TSharedPtr<SScrollBar> EntityScroll;
+	void OnUserScrolled(float InScrollOffsetFraction);
+};
diff --git a/Engine/Source/Editor/MassEntityDebugger/Private/SMassEntitiesView.cpp b/Engine/Source/Editor/MassEntityDebugger/Private/SMassEntitiesView.cpp
new file mode 100644
--- /dev/null
+++ b/Engine/Source/Editor/MassEntityDebugger/Private/SMassEntitiesView.cpp
@@ -0,0 +1,192 @@
+// Copyright Epic Games, Inc. All Rights Reserved.
+
+#include "SMassEntitiesView.h"
+#include "Widgets/Input/SButton.h"
+#include "Widgets/Layout/SScrollBox.h"
+#include "Widgets/Views/SListView.h"
+#include "Widgets/SBoxPanel.h"
+#include "MassDebugger.h"
+#include "MassDebuggerModel.h"
+#include "MassEntityQuery.h"
+#include "SMassEntitiesList.h"
+
+#define LOCTEXT_NAMESPACE "SMassDebugger"
+
+//----------------------------------------------------------------------//
+// SMassEntitiesView
+//----------------------------------------------------------------------//
+void SMassEntitiesView::Construct(const FArguments& InArgs, TSharedRef<FMassDebuggerModel> InDebuggerModel, uint32 InEntitiesViewIndex)
+{
+#if WITH_MASSENTITY_DEBUG
+	Initialize(InDebuggerModel);
+
+	EntitiesList = SNew(SMassEntitiesList, InDebuggerModel);
+	ListLabel = SNew(STextBlock);
+
+	ChildSlot
+	[
+		SNew(SVerticalBox)
+		+ SVerticalBox::Slot()
+		.AutoHeight()
+		[
+			SNew(SHorizontalBox)
+			+ SHorizontalBox::Slot()
+			.AutoWidth()
+			[
+				ListLabel.ToSharedRef()
+			]
+			+ SHorizontalBox::Slot()
+			.AutoWidth()
+			[
+				SNew(SButton)
+				.Text(LOCTEXT("RefreshEntityList", "Refresh Entities"))
+				.OnClicked(FOnClicked::CreateSP(this, &SMassEntitiesView::RefreshEntityList))
+			]
+			+ SHorizontalBox::Slot()
+			.AutoWidth()
+			[
+				SNew(SButton)
+				.Text(LOCTEXT("RefreshEntityData", "Refresh Entity Data"))
+				.OnClicked(FOnClicked::CreateSP(this, &SMassEntitiesView::RefreshEntityData))
+			]
+			+ SHorizontalBox::Slot()
+			.AutoWidth()
+			[
+				SNew(SButton)
+				.Text(LOCTEXT("ShowAllEntities", "ShowAllEntities"))
+				.OnClicked(FOnClicked::CreateSP(this, &SMassEntitiesView::ShowAllEntities))
+			]
+		]
+		+ SVerticalBox::Slot()
+		.AutoHeight()
+		[
+			EntitiesList.ToSharedRef()
+		]
+	];
+
+	InDebuggerModel->RegisterEntitiesView(SharedThis(this), InEntitiesViewIndex);
+
+#else
+	ChildSlot
+	[
+		SNew(STextBlock)
+		.Text(LOCTEXT("MassEntityDebuggingNotEnabled", "Mass Entity Debugging Not Enabled for this configuration"))
+	];
+#endif
+}
+
+void SMassEntitiesView::ShowEntities(FMassArchetypeHandle InArchetypeHandle)
+{
+#if WITH_MASSENTITY_DEBUG
+	ArchetypeHandle = InArchetypeHandle;
+	ShowEntitiesFrom = EShowEntitiesFrom::Archetype;
+	RefreshEntityList();
+#endif
+}
+
+void SMassEntitiesView::ShowEntities(const TArray<FMassEntityHandle>& InEntities)
+{
+#if WITH_MASSENTITY_DEBUG
+	ArchetypeHandle = FMassArchetypeHandle();
+	ShowEntitiesFrom = EShowEntitiesFrom::List;
+	if (EntitiesList.IsValid())
+	{
+		EntitiesList->SetEntities(InEntities);
+	}
+#endif
+}
+
+void SMassEntitiesView::ShowEntities(FMassEntityQuery& InQuery)
+{
+#if WITH_MASSENTITY_DEBUG
+	Query = InQuery;
+	ShowEntitiesFrom = EShowEntitiesFrom::Query;
+	RefreshEntityList();
+#endif
+}
+
+void SMassEntitiesView::ShowEntities(TConstArrayView<FMassEntityQuery*> InQueries)
+{
+#if WITH_MASSENTITY_DEBUG
+	Queries = InQueries;
+	ShowEntitiesFrom = EShowEntitiesFrom::QueryList;
+	RefreshEntityList();
+#endif
+}
+
+FReply SMassEntitiesView::ShowAllEntities()
+{
+#if WITH_MASSENTITY_DEBUG
+	ArchetypeHandle = FMassArchetypeHandle();
+	ShowEntitiesFrom = EShowEntitiesFrom::All;
+	RefreshEntityList();
+#endif
+	return FReply::Handled();
+}
+
+FReply SMassEntitiesView::RefreshEntityList()
+{
+#if WITH_MASSENTITY_DEBUG
+	if (!EntitiesList.IsValid() || !DebuggerModel.IsValid() || !DebuggerModel->Environment.IsValid() || !DebuggerModel->Environment->EntityManager.IsValid())
+	{ 
+		return FReply::Handled();
+	}
+	const FMassEntityManager& EntityManager = *DebuggerModel->Environment->EntityManager.Pin().Get();
+
+	switch (ShowEntitiesFrom)
+	{
+	case EShowEntitiesFrom::All:
+		{
+			TArray<FMassArchetypeHandle> Archetypes = FMassDebugger::GetAllArchetypes(EntityManager);
+			TempEntityList.Reset();
+			for (FMassArchetypeHandle& ArchHandle : Archetypes)
+			{
+				TempEntityList.Append(FMassDebugger::GetEntitiesOfArchetype(ArchHandle));
+			}
+			EntitiesList->SetEntities(TempEntityList);
+			break;
+		}
+	case EShowEntitiesFrom::Archetype:
+		TempEntityList.Reset();
+		TempEntityList.Append(FMassDebugger::GetEntitiesOfArchetype(ArchetypeHandle));
+		EntitiesList->SetEntities(TempEntityList);
+		break;
+	case EShowEntitiesFrom::Query:
+		TempEntityList.Reset();
+		TempEntityList.Append(FMassDebugger::GetEntitiesMatchingQuery(EntityManager, Query));
+		EntitiesList->SetEntities(TempEntityList);
+		break;
+	case EShowEntitiesFrom::QueryList:
+		{
+			TempEntityList.Reset();
+			for (FMassEntityQuery* QueryInList : Queries)
+			{
+				if (QueryInList)
+				{
+					TempEntityList.Append(FMassDebugger::GetEntitiesMatchingQuery(EntityManager, *QueryInList));
+				}
+			}
+			EntitiesList->SetEntities(TempEntityList);
+			break;
+		}
+	case EShowEntitiesFrom::List:
+	default:
+		break;
+	}
+#endif
+	return FReply::Handled();
+}
+
+FReply SMassEntitiesView::RefreshEntityData()
+{
+#if WITH_MASSENTITY_DEBUG
+	if (EntitiesList.IsValid())
+	{
+		EntitiesList->RefreshEntityData();
+	}
+#endif
+	return FReply::Handled();
+}
+
+#undef LOCTEXT_NAMESPACE
+
diff --git a/Engine/Source/Editor/MassEntityDebugger/Private/SMassEntitiesView.h b/Engine/Source/Editor/MassEntityDebugger/Private/SMassEntitiesView.h
new file mode 100644
--- /dev/null
+++ b/Engine/Source/Editor/MassEntityDebugger/Private/SMassEntitiesView.h
@@ -0,0 +1,66 @@
+// Copyright Epic Games, Inc. All Rights Reserved.
+#pragma once
+
+#include "SMassDebuggerViewBase.h"
+#include "Widgets/Views/SListView.h"
+#include "MassEntityHandle.h"
+#include "MassArchetypeTypes.h"
+#include "MassEntityQuery.h"
+
+class SMassEntitiesList;
+
+class SMassEntitiesView : public SMassDebuggerViewBase
+{
+public:
+	SLATE_BEGIN_ARGS(SMassEntitiesView)
+	{
+	}
+	SLATE_END_ARGS()
+
+	void Construct(const FArguments& InArgs, TSharedRef<FMassDebuggerModel> InDebuggerModel, uint32 InEntitiesViewIndex);
+
+	void ShowEntities(FMassArchetypeHandle InArchetypeHandle);
+	void ShowEntities(const TArray<FMassEntityHandle>& InEntities);
+	void ShowEntities(FMassEntityQuery& InQuery);
+	void ShowEntities(TConstArrayView<FMassEntityQuery*> InQueries);
+
+protected:
+	enum class EShowEntitiesFrom
+	{
+		List,
+		Archetype,
+		Query,
+		QueryList,
+		All
+	};
+	FReply ShowAllEntities();
+	FReply RefreshEntityList();
+	FReply RefreshEntityData();
+
+	/** Unused pure virtual function */
+	virtual void OnRefresh() override
+	{
+
+	}
+
+	/** Unused pure virtual function */
+	virtual void OnProcessorsSelected(TConstArrayView<TSharedPtr<FMassDebuggerProcessorData>> SelectedProcessors, ESelectInfo::Type SelectInfo) override
+	{
+		//@todo: highlight entities affected by the selected processor
+	}
+
+	/** Unused pure virtual function */
+	virtual void OnArchetypesSelected(TConstArrayView<TSharedPtr<FMassDebuggerArchetypeData>> SelectedArchetypes, ESelectInfo::Type SelectInfo) override
+	{
+		//@todo: highlight entities of the selected archetype
+	}
+
+	EShowEntitiesFrom ShowEntitiesFrom;
+	TSharedPtr<SMassEntitiesList> EntitiesList;
+	FMassArchetypeHandle ArchetypeHandle;
+	TSharedPtr<STextBlock> ListLabel;
+	TArray<FMassEntityHandle> TempEntityList;
+	uint32 EntitiesViewIndex;
+	FMassEntityQuery Query;
+	TConstArrayView<FMassEntityQuery*> Queries;
+};
diff --git a/Engine/Source/Editor/MassEntityDebugger/Private/SMassProcessorsWidget.cpp b/Engine/Source/Editor/MassEntityDebugger/Private/SMassProcessorsWidget.cpp
--- a/Engine/Source/Editor/MassEntityDebugger/Private/SMassProcessorsWidget.cpp
+++ b/Engine/Source/Editor/MassEntityDebugger/Private/SMassProcessorsWidget.cpp
@@ -68,14 +68,13 @@ void SMassProcessorWidget::Construct(const FArguments& InArgs, TSharedPtr<FMassD
 					[
 						SNew(SButton)
 						.ButtonStyle(FAppStyle::Get(), "NoBorder")
-						.ToolTipText(LOCTEXT("ViewEntities", "View Entities"))
 						.Content()
 						[
 							SNew(STextBlock)
 							.TextStyle(FAppStyle::Get(), "LargeText")
 							.Text(FText::FromString(ProcData.Label))
 						]
-						.OnClicked(this, &SMassProcessorWidget::HandleSelectProcessorClicked)
+						.OnClicked(FOnClicked::CreateSP(this, &SMassProcessorWidget::HandleSelectProcessorClicked))
 					]
 				]
 
@@ -96,7 +95,7 @@ void SMassProcessorWidget::Construct(const FArguments& InArgs, TSharedPtr<FMassD
 							SNew(SImage)
 							.Image(FAppStyle::Get().GetBrush("Icons.Layout"))
 						]
-						.OnClicked(this, &SMassProcessorWidget::HandleExpandTextClicked)
+						.OnClicked(FOnClicked::CreateSP(this, &SMassProcessorWidget::HandleExpandTextClicked))
 					]
 
 					+ SHorizontalBox::Slot()
@@ -110,7 +109,7 @@ void SMassProcessorWidget::Construct(const FArguments& InArgs, TSharedPtr<FMassD
 							SNew(SImage)
 							.Image(FAppStyle::Get().GetBrush("Icons.LOD"))
 						]
-						.OnClicked(this, &SMassProcessorWidget::HandleExpandGraphClicked)
+						.OnClicked(FOnClicked::CreateSP(this, &SMassProcessorWidget::HandleExpandGraphClicked))
 					]
 
 					+ SHorizontalBox::Slot()
@@ -124,7 +123,20 @@ void SMassProcessorWidget::Construct(const FArguments& InArgs, TSharedPtr<FMassD
 							SNew(SImage)
 							.Image(FAppStyle::Get().GetBrush("Icons.C++"))
 						]
-						.OnClicked(this, &SMassProcessorWidget::HandleOpenSourceLocationClicked)
+						.OnClicked(FOnClicked::CreateSP(this, &SMassProcessorWidget::HandleOpenSourceLocationClicked))
+					]
+					+ SHorizontalBox::Slot()
+					.AutoWidth()
+					[
+						SNew(SButton)
+						.ButtonStyle(FAppStyle::Get(), "NoBorder")
+						.ToolTipText(LOCTEXT("ShowAffectedEntities", "Show Affected Entities"))
+						.Content()
+						[
+							SNew(SImage)
+							.Image(FAppStyle::Get().GetBrush("Icons.Search"))
+						]
+						.OnClicked(FOnClicked::CreateSP(this, &SMassProcessorWidget::HandleShowEntitiesClicked))
 					]
 				]
 
@@ -199,6 +211,23 @@ FReply SMassProcessorWidget::HandleOpenSourceLocationClicked()
 	return FReply::Handled();
 }
 
+FReply SMassProcessorWidget::HandleShowEntitiesClicked()
+{
+	if (!ProcessorData.IsValid() || !DebuggerModel.IsValid())
+	{
+		return FReply::Handled();
+	}
+
+	const UMassProcessor* Proc = ProcessorData.Get()->Processor.Get();
+	if (Proc)
+	{
+		TConstArrayView<FMassEntityQuery*> Queries = FMassDebugger::GetProcessorQueries(*Proc);
+		DebuggerModel->ShowEntitiesView(0, Queries);
+	}
+
+	return FReply::Handled();
+}
+
 void SMassProcessorWidget::HandleFragmentSelected(FName SelectedFragment)
 {
 	Border->SetBorderImage(GetBorderByFragmentSelection());
diff --git a/Engine/Source/Editor/MassEntityDebugger/Private/SMassProcessorsWidget.h b/Engine/Source/Editor/MassEntityDebugger/Private/SMassProcessorsWidget.h
--- a/Engine/Source/Editor/MassEntityDebugger/Private/SMassProcessorsWidget.h
+++ b/Engine/Source/Editor/MassEntityDebugger/Private/SMassProcessorsWidget.h
@@ -24,6 +24,7 @@ private:
 	FReply HandleExpandGraphClicked();
 	FReply HandleSelectProcessorClicked();
 	FReply HandleOpenSourceLocationClicked();
+	FReply HandleShowEntitiesClicked();
 	void HandleFragmentSelected(FName SelectedFragment);
 	bool bIsExpandedText;
 	bool bIsExpandedGraph;
diff --git a/Engine/Source/Runtime/MassEntity/Private/MassDebugger.cpp b/Engine/Source/Runtime/MassEntity/Private/MassDebugger.cpp
--- a/Engine/Source/Runtime/MassEntity/Private/MassDebugger.cpp
+++ b/Engine/Source/Runtime/MassEntity/Private/MassDebugger.cpp
@@ -510,6 +510,18 @@ void FMassDebugger::GetQueryExecutionRequirements(const FMassEntityQuery& Query,
 	Query.ExportRequirements(OutExecutionRequirements);
 }
 
+TArray<FMassEntityHandle> FMassDebugger::GetEntitiesMatchingQuery(const FMassEntityManager& EntityManager, const FMassEntityQuery& Query)
+{
+	TArray<FMassEntityHandle> Entities;
+	TArray<FMassArchetypeHandle> Archetypes;
+	EntityManager.GetMatchingArchetypes(Query, Archetypes, 0);
+	for (FMassArchetypeHandle& ArchHandle : Archetypes)
+	{
+		Entities.Append(GetEntitiesOfArchetype(ArchHandle));
+	}
+	return Entities;
+}
+
 TArray<FMassArchetypeHandle> FMassDebugger::GetAllArchetypes(const FMassEntityManager& EntityManager)
 {
 	TArray<FMassArchetypeHandle> Archetypes;
@@ -768,11 +780,14 @@ void FMassDebugger::OutputEntityDescription(FOutputDevice& Ar, const FMassEntity
 
 void FMassDebugger::SelectEntity(const FMassEntityManager& EntityManager, const FMassEntityHandle EntityHandle)
 {
-	UE::Mass::Debug::SetDebugEntityRange(EntityHandle.Index, EntityHandle.Index);
+	if (EntityManager.IsEntityValid(EntityHandle))
+	{
+		UE::Mass::Debug::SetDebugEntityRange(EntityHandle.Index, EntityHandle.Index);
 
-	GetActiveEnvironment(EntityManager).SelectedEntity = EntityHandle;
+		GetActiveEnvironment(EntityManager).SelectedEntity = EntityHandle;
 
-	OnEntitySelectedDelegate.Broadcast(EntityManager, EntityHandle);
+		OnEntitySelectedDelegate.Broadcast(EntityManager, EntityHandle);
+	}
 }
 
 FMassEntityHandle FMassDebugger::GetSelectedEntity(const FMassEntityManager& EntityManager)
@@ -1108,6 +1123,22 @@ const UScriptStruct* FMassDebugger::GetFragmentTypeFromName(FName FragmentName)
 	return nullptr;
 }
 
+TSharedPtr<FStructOnScope> FMassDebugger::GetFragmentData(const FMassEntityManager& EntityManager, const UScriptStruct* FragmentType, FMassEntityHandle Entity)
+{
+	TSharedPtr<FMassArchetypeData> Archetype = EntityManager.DebugGetEntityStorageInterface().GetArchetypeAsShared(Entity.Index);
+	if (Archetype.IsValid())
+	{
+		void* FragmentData = Archetype->GetFragmentDataForEntity(FragmentType, Entity.Index);
+		if (FragmentData)
+		{
+			TSharedPtr<FStructOnScope> StructOnScope = MakeShared<FStructOnScope>(FragmentType);
+			CastChecked<UScriptStruct>(StructOnScope->GetStruct())->CopyScriptStruct(StructOnScope->GetStructMemory(), FragmentData);
+			return StructOnScope;
+		}
+	}
+	return nullptr;
+}
+
 void FMassDebugger::UpdateHasBreakpoint()
 {
 	bHasBreakpoint = false;
diff --git a/Engine/Source/Runtime/MassEntity/Public/MassDebugger.h b/Engine/Source/Runtime/MassEntity/Public/MassDebugger.h
--- a/Engine/Source/Runtime/MassEntity/Public/MassDebugger.h
+++ b/Engine/Source/Runtime/MassEntity/Public/MassDebugger.h
@@ -14,6 +14,7 @@
 #endif // UE_ENABLE_INCLUDE_ORDER_DEPRECATED_IN_5_6
 
 class FOutputDevice;
+class FStructOnScope;
 class UMassProcessor;
 struct FMassEntityQuery;
 struct FMassEntityManager;
@@ -169,6 +170,7 @@ struct MASSENTITY_API FMassDebugger
 
 	static UE::Mass::Debug::FQueryRequirementsView GetQueryRequirements(const FMassEntityQuery& Query);
 	static void GetQueryExecutionRequirements(const FMassEntityQuery& Query, FMassExecutionRequirements& OutExecutionRequirements);
+	static TArray<FMassEntityHandle> GetEntitiesMatchingQuery(const FMassEntityManager& EntityManager, const FMassEntityQuery& Query);
 
 	static TArray<FMassArchetypeHandle> GetAllArchetypes(const FMassEntityManager& EntitySubsystem);
 	static const FMassArchetypeCompositionDescriptor& GetArchetypeComposition(const FMassArchetypeHandle& ArchetypeHandle);
@@ -306,6 +308,11 @@ struct MASSENTITY_API FMassDebugger
 	 */
 	static const UScriptStruct* GetFragmentTypeFromName(FName FragmentName);
 
+	/**
+	 * Finds the fragment data of the specified type in the entity data. Returns nullptr if not found.
+	 */
+	static TSharedPtr<FStructOnScope> GetFragmentData(const FMassEntityManager& EntityManager, const UScriptStruct* FragmentType, FMassEntityHandle Entity);
+
 	/**
 	 * Clears all breakpoints in all environments.
 	 */
