[Mass Debugger]
Guard FMassDebugger::GetProcessorQueries call with WITH_MASSENTITY_DEBUG condition

diff --git a/Engine/Source/Editor/MassEntityDebugger/Private/SMassProcessorsWidget.cpp b/Engine/Source/Editor/MassEntityDebugger/Private/SMassProcessorsWidget.cpp
--- a/Engine/Source/Editor/MassEntityDebugger/Private/SMassProcessorsWidget.cpp
+++ b/Engine/Source/Editor/MassEntityDebugger/Private/SMassProcessorsWidget.cpp
@@ -213,6 +213,7 @@ FReply SMassProcessorWidget::HandleOpenSourceLocationClicked()
 
 FReply SMassProcessorWidget::HandleShowEntitiesClicked()
 {
+#if WITH_MASSENTITY_DEBUG
 	if (!ProcessorData.IsValid() || !DebuggerModel.IsValid())
 	{
 		return FReply::Handled();
@@ -224,7 +225,7 @@ FReply SMassProcessorWidget::HandleShowEntitiesClicked()
 		TConstArrayView<FMassEntityQuery*> Queries = FMassDebugger::GetProcessorQueries(*Proc);
 		DebuggerModel->ShowEntitiesView(0, Queries);
 	}
-
+#endif
 	return FReply::Handled();
 }
 
