[Mass Debugger]
Fix environment selection and refresh behavior

diff --git a/Engine/Source/Editor/MassEntityDebugger/Private/SMassProcessingGraphView.cpp b/Engine/Source/Editor/MassEntityDebugger/Private/SMassProcessingGraphView.cpp
--- a/Engine/Source/Editor/MassEntityDebugger/Private/SMassProcessingGraphView.cpp
+++ b/Engine/Source/Editor/MassEntityDebugger/Private/SMassProcessingGraphView.cpp
@@ -134,16 +134,11 @@ void SMassProcessingGraphView::Construct(const FArguments& InArgs, TSharedRef<FM
 
 void SMassProcessingGraphView::Display(TSharedPtr<FMassDebuggerProcessingGraph> InProcessingGraphData)
 {
-	if (!InProcessingGraphData)
-	{
-		return;
-	}
-
 	ItemsBox->ClearChildren();
 	RootNodes.Reset();
 	AllNodes.Reset();
 
-	if (InProcessingGraphData->GraphNodes.Num() == 0)
+	if (!InProcessingGraphData || InProcessingGraphData->GraphNodes.Num() == 0)
 	{
 		GraphNodesTree->RebuildList();
 		return;
diff --git a/Engine/Source/Editor/MassEntityDebugger/Private/SMassProcessingView.cpp b/Engine/Source/Editor/MassEntityDebugger/Private/SMassProcessingView.cpp
--- a/Engine/Source/Editor/MassEntityDebugger/Private/SMassProcessingView.cpp
+++ b/Engine/Source/Editor/MassEntityDebugger/Private/SMassProcessingView.cpp
@@ -64,6 +64,14 @@ void SMassProcessingView::Construct(const FArguments& InArgs, TSharedRef<FMassDe
 				})
 			.OnSelectionChanged_Lambda([this](TSharedPtr<FMassDebuggerProcessingGraph> SelectedItem, ESelectInfo::Type SelectInfo)
 				{
+					if (SelectedItem)
+					{
+						SelectedLabel = SelectedItem->Label;
+					}
+					else
+					{
+						SelectedLabel.Empty();
+					}
 					ProcessingGraphWidget->Display(SelectedItem);
 				})
 		]
@@ -81,6 +89,24 @@ void SMassProcessingView::Construct(const FArguments& InArgs, TSharedRef<FMassDe
 void SMassProcessingView::OnRefresh()
 {
 	GraphsListWidget->RequestListRefresh();
+	if (DebuggerModel.IsValid() && DebuggerModel->CachedProcessingGraphs.Num() > 0)
+	{
+		int SelectedIndex = 0;
+
+		if (!SelectedLabel.IsEmpty())
+		{
+			for (int i = 0; i < DebuggerModel->CachedProcessingGraphs.Num(); i++)
+			{
+				if (DebuggerModel->CachedProcessingGraphs[i]->Label == SelectedLabel)
+				{
+					SelectedIndex = i;
+					break;
+				}
+			}
+		}
+
+		GraphsListWidget->SetSelection(DebuggerModel->CachedProcessingGraphs[SelectedIndex]);
+	}
 }
 
 #undef LOCTEXT_NAMESPACE
diff --git a/Engine/Source/Editor/MassEntityDebugger/Private/SMassProcessingView.h b/Engine/Source/Editor/MassEntityDebugger/Private/SMassProcessingView.h
--- a/Engine/Source/Editor/MassEntityDebugger/Private/SMassProcessingView.h
+++ b/Engine/Source/Editor/MassEntityDebugger/Private/SMassProcessingView.h
@@ -28,6 +28,7 @@ protected:
 	TSharedPtr<SListView<TSharedPtr<FMassDebuggerProcessingGraph>>> GraphsListWidget;
 
 	TSharedPtr<SMassProcessingGraphView> ProcessingGraphWidget;
+	FString SelectedLabel;
 
 	FDelegateHandle OnRefreshHandle;
 };
