[Mass Debugger]
Remove "Add breakpoint" button from the debugger tab. That functionality was added to the "Entities" tab where it makes more sense.

diff --git a/Engine/Source/Editor/MassEntityDebugger/Private/SMassBreakpointsView.cpp b/Engine/Source/Editor/MassEntityDebugger/Private/SMassBreakpointsView.cpp
--- a/Engine/Source/Editor/MassEntityDebugger/Private/SMassBreakpointsView.cpp
+++ b/Engine/Source/Editor/MassEntityDebugger/Private/SMassBreakpointsView.cpp
@@ -33,8 +33,8 @@ void SMassBreakpointsView::Construct(const FArguments& InArgs, TSharedRef<FMassD
 		.Padding(5)
 		[
 			SNew(SButton)
-			.Text(LOCTEXT("AddWriteBreakpoint", "Add Write Breakpoint"))
-			.OnClicked(this, &SMassBreakpointsView::HandleAddWriteBreakpointClicked)
+			.Text(LOCTEXT("ClearAllBreakpoints", "Clear All Breakpoints"))
+			.OnClicked(this, &SMassBreakpointsView::ClearBreakpointsClicked)
 		]
 		+ SVerticalBox::Slot()
 		.AutoHeight()
@@ -42,14 +42,6 @@ void SMassBreakpointsView::Construct(const FArguments& InArgs, TSharedRef<FMassD
 		[
 			WriteBreakpointsListView.ToSharedRef()
 		]
-		+ SVerticalBox::Slot()
-		.AutoHeight()
-		.Padding(5)
-		[
-			SNew(SButton)
-			.Text(LOCTEXT("ClearAllBreakpoints", "Clear All Breakpoints"))
-			.OnClicked(this, &SMassBreakpointsView::ClearBreakpointsClicked)
-		]
 	];
 #else
 	ChildSlot
