[Mass] Alternative styling of the Mass Processor Widget

diff --git a/Engine/Source/Editor/MassEntityDebugger/Private/SMassProcessorsWidget.cpp b/Engine/Source/Editor/MassEntityDebugger/Private/SMassProcessorsWidget.cpp
--- a/Engine/Source/Editor/MassEntityDebugger/Private/SMassProcessorsWidget.cpp
+++ b/Engine/Source/Editor/MassEntityDebugger/Private/SMassProcessorsWidget.cpp
@@ -50,11 +50,11 @@ void SMassProcessorWidget::Construct(const FArguments& InArgs, TSharedPtr<FMassD
 		SAssignNew(Border, SBorder)
 		.BorderImage(GetBorderByFragmentSelection())
 		.BorderBackgroundColor(FLinearColor::Gray)
-		.Padding(3.0f)
+		.Padding(1.0f)
 		[
 			SNew(SBorder)
 			.BorderImage(FMassDebuggerStyle::GetBrush("MassDebug.Processor.InnerBackground"))
-			.Padding(5.0f)
+			.Padding(2.0f)
 			[
 				SNew(SVerticalBox)
 
@@ -71,19 +71,11 @@ void SMassProcessorWidget::Construct(const FArguments& InArgs, TSharedPtr<FMassD
 						.Content()
 						[
 							SNew(STextBlock)
-							.TextStyle(FAppStyle::Get(), "LargeText")
+							.TextStyle(FAppStyle::Get(), "NormalText")
 							.Text(FText::FromString(ProcData.Label))
 						]
 						.OnClicked(FOnClicked::CreateSP(this, &SMassProcessorWidget::HandleSelectProcessorClicked))
 					]
-				]
-
-				+ SVerticalBox::Slot()
-				.AutoHeight()
-				.HAlign(HAlign_Right)
-				[
-					SNew(SHorizontalBox)
-
 					+ SHorizontalBox::Slot()
 					.AutoWidth()
 					[
@@ -142,7 +134,7 @@ void SMassProcessorWidget::Construct(const FArguments& InArgs, TSharedPtr<FMassD
 
 				+ SVerticalBox::Slot()
 				.AutoHeight()
-				.Padding(5.0f)
+				.Padding(2.0f)
 				[
 					SAssignNew(TextBoxContainer, SBox)
 					.Visibility(bIsExpandedText ? EVisibility::Visible : EVisibility::Collapsed)
@@ -154,7 +146,7 @@ void SMassProcessorWidget::Construct(const FArguments& InArgs, TSharedPtr<FMassD
 
 				+ SVerticalBox::Slot()
 				.AutoHeight()
-				.Padding(5.0f)
+				.Padding(2.0f)
 				[
 					SAssignNew(GraphBoxContainer, SBox)
 					.Visibility(bIsExpandedGraph ? EVisibility::Visible : EVisibility::Collapsed)
