[Mass Debugger]
Rewrite entity list as Tree View
Fixed entity list view layout to fit tab bounds and improve scrolling behavior
Fixed entity fragment data refresh

diff --git a/Engine/Source/Editor/MassEntityDebugger/Private/SMassEntitiesList.cpp b/Engine/Source/Editor/MassEntityDebugger/Private/SMassEntitiesList.cpp
--- a/Engine/Source/Editor/MassEntityDebugger/Private/SMassEntitiesList.cpp
+++ b/Engine/Source/Editor/MassEntityDebugger/Private/SMassEntitiesList.cpp
@@ -29,19 +29,13 @@
 //----------------------------------------------------------------------//
 // SMassEntitiesList
 //----------------------------------------------------------------------//
-void SMassEntitiesList::Construct(const FArguments& InArgs, TSharedRef<FMassDebuggerModel> InDebuggerModel)
+void SMassEntitiesList::Construct(const SMassEntitiesList::FArguments& InArgs, TSharedRef<FMassDebuggerModel> InDebuggerModel)
 {
 #if WITH_MASSENTITY_DEBUG
-	EntityGrid = SNew(SGridPanel);
 	FragmentSelectBox = SNew(SBox);
 
 	DebuggerModel = InDebuggerModel;
 
-	EntityScroll = SNew(SScrollBar)
-		.AlwaysShowScrollbar(true)
-		.Orientation(EOrientation::Orient_Vertical);
-	EntityScroll->SetOnUserScrolled(FOnUserScrolled::CreateSP(this, &SMassEntitiesList::OnUserScrolled));
-	
 	ChildSlot
 	[
 		SNew(SBorder)
@@ -50,35 +44,32 @@ void SMassEntitiesList::Construct(const FArguments& InArgs, TSharedRef<FMassDebu
 			SNew(SVerticalBox)
 			+ SVerticalBox::Slot()
 			.AutoHeight()
+			.Padding(5.0f)
 			.HAlign(HAlign_Left)
 			[
 				FragmentSelectBox.ToSharedRef()
 			]
 			+ SVerticalBox::Slot()
+			.VAlign(VAlign_Fill)
 			.FillHeight(1.f)
 			[
-				SNew(SHorizontalBox)
-				+ SHorizontalBox::Slot()
-				.FillWidth(1.f)
-				[
-					SNew(SScrollBox)
-					.Orientation(Orient_Horizontal)
-					+SScrollBox::Slot()
-					[
-						EntityGrid.ToSharedRef()
-					]
-				]
-				+ SHorizontalBox::Slot()
-				.AutoWidth()
-				.Padding(10.0f)
-				[
-					EntityScroll.ToSharedRef()
-				]
+				SAssignNew(TreeView, STreeView<EntitiesTableRowPtr>)
+				.SelectionMode(ESelectionMode::None)
+				.TreeItemsSource(&GridRows)
+				.OnGetChildren(this, &SMassEntitiesList::TreeView_OnGetChildren)
+				.OnGenerateRow(this, &SMassEntitiesList::TreeView_OnGenerateRow)
+				.HeaderRow
+				(
+					SAssignNew(TreeViewHeaderRow, SHeaderRow)
+					.Visibility(EVisibility::Visible)
+				)
 			]
 		]
 	];
 
+	UpdateTreeColumns();
 	SetEntities(InArgs._Entities);
+
 #else
 	ChildSlot
 	[
@@ -88,21 +79,12 @@ void SMassEntitiesList::Construct(const FArguments& InArgs, TSharedRef<FMassDebu
 #endif
 }
 
-void SMassEntitiesList::OnUserScrolled(float InScrollOffsetFraction)
-{
-	float Top = float(GridRows.Num()) * (InScrollOffsetFraction);
-	int NewTop = (int)Top;
-	SetTopIndex(NewTop);
-}
 
 void SMassEntitiesList::BuildGrid()
 {
 #if WITH_MASSENTITY_DEBUG
-	EntityGrid->ClearChildren();
-
 	PopulateGridColumns();
-	PrepareRowData();
-	PopulateEntityData();
+	TreeView->RequestTreeRefresh();
 #endif
 }
 
@@ -176,21 +158,10 @@ void SMassEntitiesList::AddPropertyRecursive(TSharedPtr<SHorizontalBox> HBox, TS
 #endif
 }
 
-void SMassEntitiesList::PopulateEntityData()
-{
-#if WITH_MASSENTITY_DEBUG
-	for (int i = 0; i < NumGridRowsDrawn; i++)
-	{
-		FillGridRow(i);
-	}
-#endif
-}
-
 void SMassEntitiesList::SetEntities(const TArray<FMassEntityHandle>& InEntities)
 {
 #if WITH_MASSENTITY_DEBUG
-	EntityScroll->SetState(0.0f, (float)NumGridRowsDrawn / FMath::Max(1.f, (float)InEntities.Num()));
-	TopIndex = 0;
+
 	AvailableFragmentNames.Reset();
 	// ensure that the selected fragments are preserved even if they're not present on the new set of entites:
 	AvailableFragmentNames.Append(SelectedFragmentNames);
@@ -202,11 +173,17 @@ void SMassEntitiesList::SetEntities(const TArray<FMassEntityHandle>& InEntities)
 	{
 		const FMassEntityManager& EntityManager = *(DebuggerModel->Environment->EntityManager.Pin());
 		TSet<FMassArchetypeHandle> SearchedArchetypes;
+		const TWeakPtr<SMassEntitiesList> WeakThisPtr = StaticCastWeakPtr<SMassEntitiesList>(AsWeak());
 
 		for (int i = 0; i < InEntities.Num(); i++)
 		{
+			if (!GridRows[i].IsValid())
+			{
+				GridRows[i] = MakeShared<FGridRow>();
+			}
 			const FMassEntityHandle& Entity = InEntities[i];
-			GridRows[i].Entity = Entity;
+			GridRows[i]->Entity = Entity;
+			GridRows[i]->EntitiesList = WeakThisPtr;
 
 			FMassArchetypeHandle ArchetypeHandle = EntityManager.GetArchetypeForEntity(Entity);
 			
@@ -218,62 +195,16 @@ void SMassEntitiesList::SetEntities(const TArray<FMassEntityHandle>& InEntities)
 					AvailableFragmentNames.AddUnique(FragmentType->GetFName());
 				});
 			}
-
 		}
 	}
 
 	AvailableFragmentNames.Sort([](const FName& A, const FName& B) { return A.Compare(B) < 0; });
 
 	CreateFragmentSelectDropdown();
-	PrepareRowData();
 	BuildGrid();
 #endif
 }
 
-void SMassEntitiesList::PrepareRowData()
-{
-#if WITH_MASSENTITY_DEBUG
-	if (!DebuggerModel.IsValid() || !DebuggerModel->Environment.IsValid() || !DebuggerModel->Environment->GetEntityManager().IsValid())
-	{
-		return;
-	}
-
-	FPropertyEditorModule& PropertyEditorModule = FModuleManager::LoadModuleChecked<FPropertyEditorModule>("PropertyEditor");
-	const FMassEntityManager& EntityManager = *(DebuggerModel->Environment->GetEntityManager());
-
-	FPropertyRowGeneratorArgs GeneratorArgs;
-	GeneratorArgs.bShouldShowHiddenProperties = true;
-
-	int FragmentDisplayCount = SelectedFragmentTypes.Num();
-
-	int MaxRow = FMath::Min(TopIndex + NumGridRowsDrawn, GridRows.Num());
-	for (int RowIndex = TopIndex; RowIndex < MaxRow; RowIndex++)
-	{
-		FGridRow& Row = GridRows[RowIndex];
-		if (Row.bDirty)
-		{
-			Row.bDirty = false;
-			Row.FragmentInfo.SetNum(FragmentDisplayCount);
-
-			for (int i = 0; i < FragmentDisplayCount; i++)
-			{
-				if (Row.FragmentInfo[i].StructType != SelectedFragmentTypes[i])
-				{
-					Row.FragmentInfo[i].StructType = SelectedFragmentTypes[i];
-					Row.FragmentInfo[i].StructData = FMassDebugger::GetFragmentData(EntityManager, SelectedFragmentTypes[i], Row.Entity);
-
-					if (Row.FragmentInfo[i].StructData.IsValid())
-					{
-						Row.FragmentInfo[i].PropertyRowGenerator = PropertyEditorModule.CreatePropertyRowGenerator(GeneratorArgs);
-						Row.FragmentInfo[i].PropertyRowGenerator->SetStructure(Row.FragmentInfo[i].StructData);
-					}
-				}
-			}
-		}
-	}
-#endif
-}
-
 TSharedPtr<IPropertyHandle> FindPropertyHandle(const FProperty* Property, TArray<TSharedRef<IDetailTreeNode>>& NodesToSearch)
 {
 #if WITH_MASSENTITY_DEBUG
@@ -297,129 +228,28 @@ TSharedPtr<IPropertyHandle> FindPropertyHandle(const FProperty* Property, TArray
 	return nullptr;
 }
 
-SMassEntitiesList::FFragmentColumn* SMassEntitiesList::GetFragmentColumn(const UScriptStruct* StructType)
+void SMassEntitiesList::RefreshFragmentData()
 {
 #if WITH_MASSENTITY_DEBUG
-	for (FFragmentColumn& Column : GridColumns)
+	if (DebuggerModel.IsValid()
+		&& DebuggerModel->Environment.IsValid()
+		&& DebuggerModel->Environment->EntityManager.IsValid())
 	{
-		if (Column.StructType == StructType)
-		{
-			return &Column;
-		}
-	}
-#endif
-	return nullptr;
-}
-
-void SMassEntitiesList::FillGridRow(int RowIndex)
-{
-#if WITH_MASSENTITY_DEBUG
-	int32 GridRow = RowIndex * 2 + NumHeaderRows;
-	RowIndex += TopIndex;
-	if (RowIndex >= GridRows.Num() || RowIndex < 0)
-	{
-		return;
-	}
-	FGridRow& Row = GridRows[RowIndex];
-	FMassEntityHandle Entity = Row.Entity;
-	TWeakPtr<FMassDebuggerModel> WeakModel = DebuggerModel.ToWeakPtr();
-	TWeakPtr<const FMassEntityManager> WeakEntityManager = DebuggerModel->Environment->GetEntityManager().ToWeakPtr();
-
-	// Add the Entity Handle to the grid
-	EntityGrid->AddSlot(0, GridRow)
-	[
-		SNew(SButton)
-		.VAlign(VAlign_Fill)
-		.HAlign(HAlign_Fill)
-		.Text(FText::FromString(Entity.DebugGetDescription()))
-		.OnClicked_Lambda([WeakEntityManager, Entity]()
-		{
-			if (Entity.IsValid() && WeakEntityManager.IsValid())
-			{
-				FMassDebugger::SelectEntity(*WeakEntityManager.Pin(), Entity);
-			}
-			return FReply::Handled();
-		})
-	];
-
-	TArray<TSharedRef<IDetailTreeNode>> NodesToSearch;
-	NodesToSearch.Reserve(100);
+		const FMassEntityManager& EntityManager = *(DebuggerModel->Environment->EntityManager.Pin());
 
-	//for every selected fragment type
-	for (FGridRow::FFragmentInfo& FragmentInf : Row.FragmentInfo)
-	{
-		if (FragmentInf.PropertyRowGenerator.IsValid())
+		for (EntitiesTableRowPtr& RowPtr : GridRows)
 		{
-			const UScriptStruct* FragmentType = FragmentInf.StructType;
-
-			// the entity for this row has this fragment
-			SMassEntitiesList::FFragmentColumn* Column = GetFragmentColumn(FragmentInf.StructType);
-			if (Column)
+			if (RowPtr.IsValid())
 			{
-				EntityGrid->AddSlot(Column->ColumnNum, GridRow)
-				.HAlign(EHorizontalAlignment::HAlign_Center)
-				.VAlign(EVerticalAlignment::VAlign_Top)
-				[
-					SNew(SVerticalBox)
-					+ SVerticalBox::Slot()
-					.AutoHeight()
-					[
-						SNew(SHorizontalBox)
-						+ SHorizontalBox::Slot()
-						.AutoWidth()
-						.Padding(4.0f)
-						[
-							SNew(SButton)
-							.ButtonStyle(FAppStyle::Get(), "FlatButton")
-							.ContentPadding(4.0f)
-							.OnClicked_Lambda([Entity, FragmentType, WeakModel]()
-							{
-								TSharedPtr<FMassDebuggerModel> Model = WeakModel.Pin();
-								if (Model.IsValid() && Model->Environment.IsValid() && Model->Environment->EntityManager.IsValid())
-								{
-									FMassDebugger::SetFragmentWriteBreak(*Model->Environment->EntityManager.Pin().Get(), FragmentType, Entity);
-								}
-								return FReply::Handled();
-							})
-							[
-								SNew(SImage)
-								.Image(FAppStyle::Get().GetBrush("GenericStop"))
-								.DesiredSizeOverride(FVector2D(16, 16))
-							]
-						]
-					]
-				];
-
-				NodesToSearch = FragmentInf.PropertyRowGenerator->GetRootTreeNodes();
+				FGridRow& Row = *RowPtr;
 
-				// loop through the properties on this fragment and fill in those columns
-				for (FFragmentColumn::FPropertyColumn& PropColumn : Column->PropertyColumns)
+				for (FGridRow::FFragmentInfo& Info : Row.FragmentInfo)
 				{
-					// some of the properties we want to find might be nested in categories so we need to search the tree
-					TSharedPtr<IPropertyHandle> PropertyHandle = FindPropertyHandle(PropColumn.Property, NodesToSearch);
-					if (PropertyHandle.IsValid())
+					if (Info.StructData.IsValid())
 					{
-						// Found a top-level property on the fragment
-
-						TSharedPtr<SHorizontalBox> HBox = SNew(SHorizontalBox);
-						TSharedPtr<SVerticalBox> VBox = SNew(SVerticalBox);
-						HBox->AddSlot()
-						.HAlign(EHorizontalAlignment::HAlign_Center)
-						.VAlign(EVerticalAlignment::VAlign_Top)
-						.AutoWidth()
-						[
-							VBox.ToSharedRef()
-						];
-
-						AddPropertyRecursive(HBox, VBox, PropertyHandle);
-
-						EntityGrid->AddSlot(PropColumn.ColumnNum, GridRow)
-						[
-							HBox.ToSharedRef()
-						];
+						FMassDebugger::GetFragmentData(EntityManager, Info.StructType, Row.Entity, Info.StructData);
 					}
 				}
-
 			}
 		}
 	}
@@ -429,16 +259,9 @@ void SMassEntitiesList::FillGridRow(int RowIndex)
 void SMassEntitiesList::PopulateGridColumns()
 {
 #if WITH_MASSENTITY_DEBUG
-	GridColumns.Reset();
-
-	// Column 0 reserved for entity handle
-	EntityGrid->AddSlot(0,1)
-	[
-		SNew(STextBlock)
-		.Text(LOCTEXT("EntityHandle", "Entity Handle"))
-	];
+	Columns.Reset();
+	ColumnIndexByID.Reset();
 
-	int CurColumn = 1;
 	for (FName FragmentName : SelectedFragmentNames)
 	{
 		const UScriptStruct* FragmentStructType = FMassDebugger::GetFragmentTypeFromName(FragmentName);
@@ -447,105 +270,27 @@ void SMassEntitiesList::PopulateGridColumns()
 			continue;
 		}
 
-		FFragmentColumn& FragmentColumn = GridColumns.AddDefaulted_GetRef();
-		FragmentColumn.ColumnNum = CurColumn++;
-		FragmentColumn.StructType = FragmentStructType;
+		FMassEntitiesListColumn& Column = Columns.AddDefaulted_GetRef();
+		Column.ColumnID = FName(FragmentStructType->GetName());
+		Column.StructType = FragmentStructType;
+		Column.ColumnLabel = FragmentStructType->GetName();
+		ColumnIndexByID.Add(Column.ColumnID, Columns.Num() - 1);
 
 		for (TFieldIterator<FProperty> It(FragmentStructType); It; ++It)
 		{
 			FProperty* Property = *It;
 
-			FFragmentColumn::FPropertyColumn& Column = FragmentColumn.PropertyColumns.AddDefaulted_GetRef();
-			Column.ColumnNum = CurColumn;
-			Column.Property = Property;
-
-			EntityGrid->AddSlot(CurColumn, 1)
-			.HAlign(EHorizontalAlignment::HAlign_Left)
-			.VAlign(EVerticalAlignment::VAlign_Center)
-			.Padding(6.0f)
-			[
-				SNew(STextBlock)
-				.Text(FText::FromString(Property->GetName()))
-			];
-			++CurColumn;
+			FMassEntitiesListColumn& PropertyColumn = Columns.AddDefaulted_GetRef();
+			PropertyColumn.StructType = FragmentStructType;
+			PropertyColumn.Property = Property;
+			PropertyColumn.ColumnID = FName(FString::Printf(TEXT("%s_%s"), *FragmentStructType->GetName(), *Property->GetName()));
+			PropertyColumn.ColumnLabel = Property->GetName();
+			ColumnIndexByID.Add(PropertyColumn.ColumnID, Columns.Num() - 1);
 		}
-
-		SGridPanel::FSlot::FSlotArguments& Slot = EntityGrid->AddSlot(FragmentColumn.ColumnNum, 0)
-		.ColumnSpan(CurColumn - FragmentColumn.ColumnNum + 1)
-		.HAlign(EHorizontalAlignment::HAlign_Center)
-		.VAlign(EVerticalAlignment::VAlign_Center)
-		.Padding(6.0f)
-		[
-			SNew(STextBlock)
-			.Text(FText::FromName(FragmentName))
-		];
-
-		// Separator
-		EntityGrid->AddSlot(++CurColumn, 1)
-		.RowSpan(100)
-		[
-			SNew(SSeparator)
-			.Orientation(Orient_Vertical)
-		];
-		++CurColumn;
 	}
-#endif
-}
-
-void SMassEntitiesList::SetTopIndex(int32 NewTopIndex)
-{
-#if WITH_MASSENTITY_DEBUG
-	NewTopIndex = FMath::Clamp(NewTopIndex, 0, FMath::Max(GridRows.Num() - NumGridRowsDrawn, 0));
-	int32 Delta = NewTopIndex - TopIndex;
-	if (Delta == 0)
-	{
-		return;
-	}
-	const int32 MaxGridRow = NumHeaderRows + (2 * NumGridRowsDrawn);
-
-	EntityScroll->SetState((float)NewTopIndex / (float)GridRows.Num(), (float)NumGridRowsDrawn / FMath::Max(1.f, (float)GridRows.Num()));;
 
-	if (TPanelChildren<SGridPanel::FSlot>* ChildrenPtr = static_cast<TPanelChildren<SGridPanel::FSlot>*>(EntityGrid->GetChildren()))
-	{
-		TPanelChildren<SGridPanel::FSlot>& Children = *ChildrenPtr;
-
-		for (int32 SlotIndex = Children.Num() - 1; SlotIndex >= 0; --SlotIndex)
-		{
-			SGridPanel::FSlot& Slot = Children[SlotIndex];
-			int32 CurRow = Slot.GetRow();
-			int32 NewRow = CurRow - (2 * Delta);
-			if (CurRow >= NumHeaderRows)
-			{
-				if (NewRow >= NumHeaderRows && NewRow < MaxGridRow)
-				{
-					Slot.SetRow(NewRow);
-				}
-				else
-				{
-					Children.RemoveAt(SlotIndex);
-				}
-			}
-		}
-	}
+	UpdateTreeColumns();
 
-	TopIndex = NewTopIndex;
-	int FillFrom;
-	int FillTo;
-	if (Delta > 0)
-	{
-		FillFrom = FMath::Max(NumGridRowsDrawn - Delta, 0);
-		FillTo = NumGridRowsDrawn;
-	}
-	else
-	{
-		FillFrom = 0;
-		FillTo = FMath::Min(NumGridRowsDrawn, -Delta);
-	}
-	for (int i = FillFrom; i < FillTo; i++)
-	{
-		FillGridRow(i);
-	}
-	
 #endif
 }
 
@@ -577,9 +322,9 @@ void SMassEntitiesList::OnFragmentCheckStateChanged(ECheckBoxState NewState, FNa
 		SelectedFragmentTypes.Add(FMassDebugger::GetFragmentTypeFromName(SelectedFragmentName));
 	}
 	
-	for (FGridRow& Row : GridRows)
+	for (EntitiesTableRowPtr& Row : GridRows)
 	{
-		Row.bDirty = true;
+		Row->bDirty = true;
 	}
 
 	BuildGrid();
@@ -649,7 +394,296 @@ void SMassEntitiesList::CreateFragmentSelectDropdown()
 
 void SMassEntitiesList::RefreshEntityData()
 {
-	BuildGrid();
+	RefreshFragmentData();
+}
+
+void SMassEntitiesList::UpdateTreeColumns()
+{
+	const FName HandleName = FName(TEXT("EntityHandle"));
+
+	const TIndirectArray<SHeaderRow::FColumn>& TreeColumns = TreeViewHeaderRow->GetColumns();
+	bool Changed = false;
+
+	if (TreeColumns.Num() == 0)
+	{
+		SHeaderRow::FColumn::FArguments ColumnArgs;
+		ColumnArgs
+		.ColumnId(HandleName)
+		.DefaultLabel(LOCTEXT("MassEntityHandle", "Entity Handle"))
+		.ToolTipText(LOCTEXT("MassEntityHandle", "Entity Handle"))
+		.HAlignHeader(HAlign_Left)
+		.VAlignHeader(VAlign_Center)
+		.HAlignCell(HAlign_Fill)
+		.VAlignCell(VAlign_Fill)
+		.InitialSortMode(EColumnSortMode::Ascending)
+		.FillWidth(50.f)
+		.FillSized(60)
+		.HeaderContent()
+		[
+			SNew(SBox)
+			.Padding(FMargin(3.0f))
+			.VAlign(VAlign_Center)
+			[
+				SNew(STextBlock)
+				.Text(LOCTEXT("MassEntityHandle", "Entity Handle"))
+			]
+		];
+		TreeViewHeaderRow->AddColumn(ColumnArgs);
+		Changed = true;
+	}
+
+	for (int i = 0; i < Columns.Num(); i++)
+	{
+		FMassEntitiesListColumn& Column = Columns[i];
+		int TreeViewColumnIndex = i + 1; // one extra column for the handle
+
+		SHeaderRow::FColumn::FArguments ColumnArgs;
+		ColumnArgs
+		.ColumnId(Column.ColumnID)
+		.DefaultLabel(FText::FromString(*Column.ColumnLabel))
+		.HAlignHeader(HAlign_Left)
+		.VAlignHeader(VAlign_Center)
+		.HAlignCell(HAlign_Fill)
+		.VAlignCell(VAlign_Fill)
+		.InitialSortMode(EColumnSortMode::Ascending)
+		.FillWidth(100.f)
+		.HeaderContent()
+		[
+			SNew(SBox)
+			.Padding(FMargin(3.0f))
+			.VAlign(VAlign_Center)
+			[
+				SNew(STextBlock)
+				.Text(FText::FromString(*Column.ColumnLabel))
+			]
+		];
+
+		if (TreeColumns.Num() > TreeViewColumnIndex)
+		{
+			if (TreeColumns[TreeViewColumnIndex].ColumnId != Column.ColumnID)
+			{
+				// insert
+				TreeViewHeaderRow->InsertColumn(ColumnArgs, TreeViewColumnIndex);
+				Changed = true;
+			}
+		}
+		else
+		{
+			TreeViewHeaderRow->AddColumn(ColumnArgs);
+			Changed = true;
+		}
+	}
+
+	// prune removed columns:
+	for (int i = TreeColumns.Num() - 1; i >= (Columns.Num() + 1); i--)
+	{
+		TreeViewHeaderRow->RemoveColumn(TreeColumns[i].ColumnId);
+		Changed = true;
+	}
+
+	if (Changed)
+	{
+		TreeView->RebuildList();
+	}
+}
+
+void SMassEntitiesList::TreeView_OnGetChildren(SMassEntitiesList::EntitiesTableRowPtr InParent,
+	TArray<SMassEntitiesList::EntitiesTableRowPtr>& OutChildren)
+{
+	return;
+}
+
+TSharedRef<ITableRow> SMassEntitiesList::TreeView_OnGenerateRow(SMassEntitiesList::EntitiesTableRowPtr Row,
+	const TSharedRef<STableViewBase>& OwnerTable)
+{
+	TSharedRef<ITableRow> TableRow =
+		SNew(SMassEntitiesList::SEntitiesTableRow, OwnerTable)
+		.EntitiesTableRow(Row);
+
+	return TableRow;
+}
+
+void SMassEntitiesList::SEntitiesTableRow::Construct(const FArguments& InArgs, const TSharedRef<STableViewBase>& InOwnerTableView)
+{
+	TableRowPtr = InArgs._EntitiesTableRow;
+	FGridRow& Row = *TableRowPtr;
+
+	if (TableRowPtr.IsValid())
+	{
+		SMassEntitiesList& OwnerList = *Row.EntitiesList.Pin();
+		if (OwnerList.DebuggerModel.IsValid()
+			&& OwnerList.DebuggerModel->Environment.IsValid()
+			&& OwnerList.DebuggerModel->Environment->EntityManager.IsValid())
+		{
+			const FMassEntityManager& EntityManager = *(OwnerList.DebuggerModel->Environment->EntityManager.Pin());
+			FPropertyRowGeneratorArgs GeneratorArgs;
+			GeneratorArgs.bShouldShowHiddenProperties = true;
+			FPropertyEditorModule& PropertyEditorModule = FModuleManager::LoadModuleChecked<FPropertyEditorModule>("PropertyEditor");
+
+			Row.bDirty = false;
+			int32 FragmentDisplayCount = OwnerList.SelectedFragmentTypes.Num();
+			Row.FragmentInfo.SetNum(FragmentDisplayCount);
+
+			for (int i = 0; i < FragmentDisplayCount; i++)
+			{
+				if (Row.FragmentInfo[i].StructType != OwnerList.SelectedFragmentTypes[i])
+				{
+					Row.FragmentInfo[i].StructType = OwnerList.SelectedFragmentTypes[i];
+					Row.FragmentInfo[i].StructData = FMassDebugger::GetFragmentData(EntityManager, OwnerList.SelectedFragmentTypes[i], Row.Entity);
+					if (Row.FragmentInfo[i].StructData.IsValid())
+					{
+						Row.FragmentInfo[i].PropertyRowGenerator = PropertyEditorModule.CreatePropertyRowGenerator(GeneratorArgs);
+						Row.FragmentInfo[i].PropertyRowGenerator->SetStructure(Row.FragmentInfo[i].StructData);
+					}
+					
+				}
+			}
+		}
+	}
+
+	SetEnabled(true);
+	Super::Construct(Super::FArguments(), InOwnerTableView);
+}
+
+TSharedRef<SWidget> SMassEntitiesList::SEntitiesTableRow::GenerateWidgetForColumn(const FName& InColumnName)
+{
+	if (TableRowPtr && TableRowPtr->EntitiesList.IsValid())
+	{
+		SMassEntitiesList& OwnerList = *TableRowPtr->EntitiesList.Pin();
+		if (!OwnerList.DebuggerModel->Environment.IsValid()
+			|| !OwnerList.DebuggerModel->Environment->GetEntityManager().IsValid())
+		{
+			return SNullWidget::NullWidget;
+		}
+
+		if (InColumnName == FName(TEXT("EntityHandle")))
+		{
+			FMassEntityHandle Entity = TableRowPtr->Entity;
+			TWeakPtr<const FMassEntityManager> WeakEntityManager = OwnerList.DebuggerModel->Environment->GetEntityManager().ToWeakPtr();
+
+			return SNew(SButton)
+			.VAlign(VAlign_Fill)
+			.HAlign(HAlign_Fill)
+			.Text(FText::FromString(Entity.DebugGetDescription()))
+			.OnClicked_Lambda([WeakEntityManager, Entity]()
+			{
+				if (Entity.IsValid() && WeakEntityManager.IsValid())
+				{
+					FMassDebugger::SelectEntity(*WeakEntityManager.Pin(), Entity);
+				}
+				return FReply::Handled();
+			});
+		}
+
+		int32* ColumnIndex = OwnerList.ColumnIndexByID.Find(InColumnName);
+
+		if (ColumnIndex && *ColumnIndex >= 0 && *ColumnIndex < OwnerList.Columns.Num())
+		{
+			FMassEntitiesListColumn& Column = OwnerList.Columns[*ColumnIndex];
+
+			if (Column.StructType)
+			{
+				for (SMassEntitiesList::FGridRow::FFragmentInfo& Info : TableRowPtr->FragmentInfo)
+				{
+					if (Info.StructType == Column.StructType)
+					{
+						if (Column.Property)
+						{
+							return GenerateDataWidget(Column.Property, Info);
+						}
+						else
+						{
+							return GenerateBreakpointWidget(Info);
+						}
+					}
+				}
+			}
+		}
+	}
+
+	return SNullWidget::NullWidget;
+}
+
+TSharedRef<SWidget> SMassEntitiesList::SEntitiesTableRow::GenerateBreakpointWidget(SMassEntitiesList::FGridRow::FFragmentInfo& Info)
+{
+	FMassEntityHandle Entity = TableRowPtr->Entity;
+	SMassEntitiesList& OwnerList = *TableRowPtr->EntitiesList.Pin();
+	TWeakPtr<FMassDebuggerModel> WeakModel = OwnerList.DebuggerModel.ToWeakPtr();
+	if (!OwnerList.DebuggerModel->Environment.IsValid()
+		|| !OwnerList.DebuggerModel->Environment->GetEntityManager().IsValid())
+	{
+		return SNullWidget::NullWidget;
+	}
+	TWeakPtr<const FMassEntityManager> WeakEntityManager = OwnerList.DebuggerModel->Environment->GetEntityManager().ToWeakPtr();
+
+	const UScriptStruct* FragmentType = Info.StructType;
+
+	return SNew(SVerticalBox)
+	+ SVerticalBox::Slot()
+	.AutoHeight()
+	[
+		SNew(SHorizontalBox)
+		+ SHorizontalBox::Slot()
+		.AutoWidth()
+		.Padding(4.0f)
+		[
+			SNew(SButton)
+			.ButtonStyle(FAppStyle::Get(), "FlatButton")
+			.ContentPadding(4.0f)
+			.ToolTipText(LOCTEXT("SetWriteBreakpoint", "Set Write Breakpoint"))
+			.OnClicked_Lambda([Entity, FragmentType, WeakModel]()
+			{
+				TSharedPtr<FMassDebuggerModel> Model = WeakModel.Pin();
+				if (Model.IsValid() && Model->Environment.IsValid() && Model->Environment->EntityManager.IsValid())
+				{
+					FMassDebugger::SetFragmentWriteBreak(*Model->Environment->EntityManager.Pin(), FragmentType, Entity);
+				}
+				return FReply::Handled();
+			})
+			[
+				SNew(SImage)
+				.Image(FAppStyle::Get().GetBrush("GenericStop"))
+				.DesiredSizeOverride(FVector2D(16, 16))
+			]
+		]
+	];
+}
+
+TSharedRef<SWidget> SMassEntitiesList::SEntitiesTableRow::GenerateDataWidget(const FProperty* Property, SMassEntitiesList::FGridRow::FFragmentInfo& Info)
+{
+	if (Info.PropertyRowGenerator.IsValid())
+	{
+		SMassEntitiesList& OwnerList = *TableRowPtr->EntitiesList.Pin();
+
+		if (!Property)
+		{
+			return SNullWidget::NullWidget;
+		}
+		
+		OwnerList.NodesToSearch.Reserve(100);
+		OwnerList.NodesToSearch = Info.PropertyRowGenerator->GetRootTreeNodes();
+
+		// some of the properties we want to find might be nested in categories so we need to search the tree
+		TSharedPtr<IPropertyHandle> PropertyHandle = FindPropertyHandle(Property, OwnerList.NodesToSearch);
+		OwnerList.NodesToSearch.Reset();
+		if (PropertyHandle.IsValid())
+		{
+			TSharedPtr<SHorizontalBox> HBox = SNew(SHorizontalBox);
+			// Found a top-level property on the fragment
+			TSharedPtr<SVerticalBox> VBox = SNew(SVerticalBox);
+			HBox->AddSlot()
+			.HAlign(EHorizontalAlignment::HAlign_Center)
+			.VAlign(EVerticalAlignment::VAlign_Top)
+			.AutoWidth()
+			[
+				VBox.ToSharedRef()
+			];
+			OwnerList.AddPropertyRecursive(HBox, VBox, PropertyHandle);
+
+			return HBox.ToSharedRef();
+		}
+	}
+	return SNullWidget::NullWidget;
 }
 
 #undef LOCTEXT_NAMESPACE
\ No newline at end of file
diff --git a/Engine/Source/Editor/MassEntityDebugger/Private/SMassEntitiesList.h b/Engine/Source/Editor/MassEntityDebugger/Private/SMassEntitiesList.h
--- a/Engine/Source/Editor/MassEntityDebugger/Private/SMassEntitiesList.h
+++ b/Engine/Source/Editor/MassEntityDebugger/Private/SMassEntitiesList.h
@@ -2,9 +2,12 @@
 #pragma once
 
 #include "Widgets/SCompoundWidget.h"
+#include "Widgets/Views/STableRow.h"
+#include "Widgets/Views/STreeView.h"
 #include "MassDebugger.h"
 #include "MassDebuggerModel.h"
 
+class IDetailTreeNode;
 class IPropertyHandle;
 class IPropertyRowGenerator;
 class FStructOnScope;
@@ -18,7 +21,7 @@ class SMassEntitiesList : public SCompoundWidget
 {
 public:
 
-	SLATE_BEGIN_ARGS(SMassEntitiesList){}
+	SLATE_BEGIN_ARGS(SMassEntitiesList){}	
 		SLATE_ARGUMENT(TArray<FMassEntityHandle>, Entities)
 	SLATE_END_ARGS()
 
@@ -27,16 +30,15 @@ public:
 	void RefreshEntityData();
 protected:
 
-	struct FFragmentColumn
+
+
+
+	struct FMassEntitiesListColumn
 	{
-		int32 ColumnNum;
 		const UScriptStruct* StructType = nullptr;
-		struct FPropertyColumn
-		{
-			int32 ColumnNum;
-			const FProperty* Property = nullptr;
-		};
-		TArray<FPropertyColumn> PropertyColumns;
+		const FProperty* Property = nullptr;
+		FString ColumnLabel;
+		FName ColumnID;
 	};
 
 	struct FGridRow
@@ -52,23 +54,41 @@ protected:
 		
 		TArray<FFragmentInfo> FragmentInfo;
 		FMassEntityHandle Entity;
+		TWeakPtr<SMassEntitiesList> EntitiesList;
 		bool bDirty = true;
 	};
 
-	static const int32 NumHeaderRows = 2;
-	static const int32 NumGridRowsDrawn = 32;
-	TArray<FGridRow> GridRows;
-	TArray<FFragmentColumn> GridColumns;
 
-	int32 TopIndex;
+
+public:
+
+	using EntitiesTableRowPtr = TSharedPtr<FGridRow>;
+
+	class SEntitiesTableRow : public SMultiColumnTableRow<EntitiesTableRowPtr>
+	{
+		using Super = SMultiColumnTableRow<EntitiesTableRowPtr>;
+	public:
+		SLATE_BEGIN_ARGS(SEntitiesTableRow) {}
+			SLATE_ARGUMENT(EntitiesTableRowPtr, EntitiesTableRow)
+		SLATE_END_ARGS()
+
+		void Construct(const FArguments& InArgs, const TSharedRef<STableViewBase>& InOwnerTableView);
+
+		virtual TSharedRef<SWidget> GenerateWidgetForColumn(const FName& InColumnName) override;
+
+	private:
+		EntitiesTableRowPtr TableRowPtr;
+		TSharedRef<SWidget> GenerateBreakpointWidget(FGridRow::FFragmentInfo& Info);
+		TSharedRef<SWidget> GenerateDataWidget(const FProperty* Property, FGridRow::FFragmentInfo& Info);
+	};
+
+	TArray<EntitiesTableRowPtr> GridRows;
+	TArray<FMassEntitiesListColumn> Columns;
+	TMap<FName, int32> ColumnIndexByID;
+
 	void BuildGrid();
-	void PopulateEntityData();
 	void PopulateGridColumns();
-	void PopulateRow(const FMassEntityManager& EntityManager, int RowIndex);
-	void PrepareRowData();
-	void FillGridRow(int Row);
-	FFragmentColumn* GetFragmentColumn(const UScriptStruct* StructType);
-	void SetTopIndex(int32 NewTopIndex);
+	void RefreshFragmentData();
 
 	void AddPropertyRecursive(TSharedPtr<SHorizontalBox> HBox, TSharedPtr<SVerticalBox> VBox, TSharedPtr<IPropertyHandle> Prop, bool bShowName = false);
 	TArray<FName> AvailableFragmentNames;
@@ -77,14 +97,18 @@ protected:
 	TSharedPtr<SWidget> FragmentSelectMenu;
 	FReply OnClearAllSelectedFragmentsClicked();
 
-	TSharedPtr<SGridPanel> EntityGrid;
 	TSharedPtr<SBox> FragmentSelectBox;
 
 	void OnFragmentCheckStateChanged(ECheckBoxState NewState, FName FragmentName);
 	ECheckBoxState GetFragmentCheckState(FName FragmentName) const;
 	void CreateFragmentSelectDropdown();
 	TSharedPtr<FMassDebuggerModel> DebuggerModel;
-	
-	TSharedPtr<SScrollBar> EntityScroll;
-	void OnUserScrolled(float InScrollOffsetFraction);
+	void UpdateTreeColumns();
+
+	void TreeView_OnGetChildren(SMassEntitiesList::EntitiesTableRowPtr InParent, TArray<SMassEntitiesList::EntitiesTableRowPtr>& OutChildren);
+	TSharedRef<ITableRow> TreeView_OnGenerateRow(EntitiesTableRowPtr RowPtr, const TSharedRef<STableViewBase>& OwnerTable);
+
+	TSharedPtr<STreeView<EntitiesTableRowPtr>> TreeView;
+	TSharedPtr<SHeaderRow> TreeViewHeaderRow;
+	TArray<TSharedRef<IDetailTreeNode>> NodesToSearch;
 };
diff --git a/Engine/Source/Editor/MassEntityDebugger/Private/SMassEntitiesView.cpp b/Engine/Source/Editor/MassEntityDebugger/Private/SMassEntitiesView.cpp
--- a/Engine/Source/Editor/MassEntityDebugger/Private/SMassEntitiesView.cpp
+++ b/Engine/Source/Editor/MassEntityDebugger/Private/SMassEntitiesView.cpp
@@ -58,7 +58,7 @@ void SMassEntitiesView::Construct(const FArguments& InArgs, TSharedRef<FMassDebu
 			]
 		]
 		+ SVerticalBox::Slot()
-		.AutoHeight()
+		.FillHeight(1.f)
 		[
 			EntitiesList.ToSharedRef()
 		]
diff --git a/Engine/Source/Runtime/MassEntity/Private/MassDebugger.cpp b/Engine/Source/Runtime/MassEntity/Private/MassDebugger.cpp
--- a/Engine/Source/Runtime/MassEntity/Private/MassDebugger.cpp
+++ b/Engine/Source/Runtime/MassEntity/Private/MassDebugger.cpp
@@ -1182,6 +1182,16 @@ const UScriptStruct* FMassDebugger::GetFragmentTypeFromName(FName FragmentName)
 }
 
 TSharedPtr<FStructOnScope> FMassDebugger::GetFragmentData(const FMassEntityManager& EntityManager, const UScriptStruct* FragmentType, FMassEntityHandle Entity)
+{
+	TSharedPtr<FStructOnScope> StructOnScope = MakeShared<FStructOnScope>(FragmentType);
+	if (GetFragmentData(EntityManager, FragmentType, Entity, StructOnScope))
+	{
+		return StructOnScope;
+	}
+	return nullptr;
+}
+
+bool FMassDebugger::GetFragmentData(const FMassEntityManager& EntityManager, const UScriptStruct* FragmentType, FMassEntityHandle Entity, TSharedPtr<FStructOnScope>& OutStructData)
 {
 	TSharedPtr<FMassArchetypeData> Archetype = EntityManager.DebugGetEntityStorageInterface().GetArchetypeAsShared(Entity.Index);
 	if (Archetype.IsValid())
@@ -1189,12 +1199,17 @@ TSharedPtr<FStructOnScope> FMassDebugger::GetFragmentData(const FMassEntityManag
 		void* FragmentData = Archetype->GetFragmentDataForEntity(FragmentType, Entity.Index);
 		if (FragmentData)
 		{
-			TSharedPtr<FStructOnScope> StructOnScope = MakeShared<FStructOnScope>(FragmentType);
-			CastChecked<UScriptStruct>(StructOnScope->GetStruct())->CopyScriptStruct(StructOnScope->GetStructMemory(), FragmentData);
-			return StructOnScope;
+			const UStruct* AsUStruct = FragmentType;
+			if (OutStructData->GetStruct() != AsUStruct)
+			{
+				OutStructData->Initialize(FragmentType);
+			}
+			
+			CastChecked<UScriptStruct>(OutStructData->GetStruct())->CopyScriptStruct(OutStructData->GetStructMemory(), FragmentData);
+			return true;
 		}
 	}
-	return nullptr;
+	return false;
 }
 
 void FMassDebugger::UpdateHasBreakpoint()
@@ -1224,8 +1239,7 @@ FMassDebugger::FEnvironment::FEnvironment(const FMassEntityManager& InEntityMana
 	: EntityManager(InEntityManager.AsWeak())
 {
 #if UE_MASS_TRACE_ENABLED
-	TWeakPtr<const FMassEntityManager> WeakEntityManager = EntityManager;
-	TraceStartedDelegateHandle = FTraceAuxiliary::OnTraceStarted.AddLambda([WeakEntityManager](FTraceAuxiliary::EConnectionType TraceType, const FString& TraceDestination)
+	TraceStartedDelegateHandle = FTraceAuxiliary::OnTraceStarted.AddLambda([WeakEntityManager = EntityManager](FTraceAuxiliary::EConnectionType TraceType, const FString& TraceDestination)
 	{
 		if (!WeakEntityManager.IsValid())
 		{
diff --git a/Engine/Source/Runtime/MassEntity/Public/MassDebugger.h b/Engine/Source/Runtime/MassEntity/Public/MassDebugger.h
--- a/Engine/Source/Runtime/MassEntity/Public/MassDebugger.h
+++ b/Engine/Source/Runtime/MassEntity/Public/MassDebugger.h
@@ -330,6 +330,12 @@ struct MASSENTITY_API FMassDebugger
 	 */
 	static TSharedPtr<FStructOnScope> GetFragmentData(const FMassEntityManager& EntityManager, const UScriptStruct* FragmentType, FMassEntityHandle Entity);
 
+	/**
+	 * Finds the fragment data of the specified type in the entity data. Returns nullptr if not found.
+	 */
+	static bool GetFragmentData(const FMassEntityManager& EntityManager, const UScriptStruct* FragmentType, FMassEntityHandle Entity, TSharedPtr<FStructOnScope>& OutStructData);
+
+
 	/**
 	 * Clears all breakpoints in all environments.
 	 */
