[Mass Debugger]
Fix debugger compilation when  WITH_MASSENTITY_DEBUG is disabled

diff --git a/Engine/Source/Editor/MassEntityDebugger/Private/SMassEntitiesList.cpp b/Engine/Source/Editor/MassEntityDebugger/Private/SMassEntitiesList.cpp
--- a/Engine/Source/Editor/MassEntityDebugger/Private/SMassEntitiesList.cpp
+++ b/Engine/Source/Editor/MassEntityDebugger/Private/SMassEntitiesList.cpp
@@ -399,6 +399,7 @@ void SMassEntitiesList::RefreshEntityData()
 
 void SMassEntitiesList::UpdateTreeColumns()
 {
+#if WITH_MASSENTITY_DEBUG
 	const FName HandleName = FName(TEXT("EntityHandle"));
 
 	const TIndirectArray<SHeaderRow::FColumn>& TreeColumns = TreeViewHeaderRow->GetColumns();
@@ -485,6 +486,7 @@ void SMassEntitiesList::UpdateTreeColumns()
 	{
 		TreeView->RebuildList();
 	}
+#endif //WITH_MASSENTITY_DEBUG
 }
 
 void SMassEntitiesList::TreeView_OnGetChildren(SMassEntitiesList::EntitiesTableRowPtr InParent,
@@ -505,6 +507,7 @@ TSharedRef<ITableRow> SMassEntitiesList::TreeView_OnGenerateRow(SMassEntitiesLis
 
 void SMassEntitiesList::SEntitiesTableRow::Construct(const FArguments& InArgs, const TSharedRef<STableViewBase>& InOwnerTableView)
 {
+#if WITH_MASSENTITY_DEBUG
 	TableRowPtr = InArgs._EntitiesTableRow;
 	FGridRow& Row = *TableRowPtr;
 
@@ -542,11 +545,13 @@ void SMassEntitiesList::SEntitiesTableRow::Construct(const FArguments& InArgs, c
 	}
 
 	SetEnabled(true);
+#endif //WITH_MASSENTITY_DEBUG
 	Super::Construct(Super::FArguments(), InOwnerTableView);
 }
 
 TSharedRef<SWidget> SMassEntitiesList::SEntitiesTableRow::GenerateWidgetForColumn(const FName& InColumnName)
 {
+#if WITH_MASSENTITY_DEBUG
 	if (TableRowPtr && TableRowPtr->EntitiesList.IsValid())
 	{
 		SMassEntitiesList& OwnerList = *TableRowPtr->EntitiesList.Pin();
@@ -600,12 +605,13 @@ TSharedRef<SWidget> SMassEntitiesList::SEntitiesTableRow::GenerateWidgetForColum
 			}
 		}
 	}
-
+#endif //WITH_MASSENTITY_DEBUG
 	return SNullWidget::NullWidget;
 }
 
 TSharedRef<SWidget> SMassEntitiesList::SEntitiesTableRow::GenerateBreakpointWidget(SMassEntitiesList::FGridRow::FFragmentInfo& Info)
 {
+#if WITH_MASSENTITY_DEBUG
 	FMassEntityHandle Entity = TableRowPtr->Entity;
 	SMassEntitiesList& OwnerList = *TableRowPtr->EntitiesList.Pin();
 	TWeakPtr<FMassDebuggerModel> WeakModel = OwnerList.DebuggerModel.ToWeakPtr();
@@ -647,10 +653,14 @@ TSharedRef<SWidget> SMassEntitiesList::SEntitiesTableRow::GenerateBreakpointWidg
 			]
 		]
 	];
+#else //WITH_MASSENTITY_DEBUG
+	return SNullWidget::NullWidget;
+#endif //WITH_MASSENTITY_DEBUG
 }
 
 TSharedRef<SWidget> SMassEntitiesList::SEntitiesTableRow::GenerateDataWidget(const FProperty* Property, SMassEntitiesList::FGridRow::FFragmentInfo& Info)
 {
+#if WITH_MASSENTITY_DEBUG
 	if (Info.PropertyRowGenerator.IsValid())
 	{
 		SMassEntitiesList& OwnerList = *TableRowPtr->EntitiesList.Pin();
@@ -683,6 +693,7 @@ TSharedRef<SWidget> SMassEntitiesList::SEntitiesTableRow::GenerateDataWidget(con
 			return HBox.ToSharedRef();
 		}
 	}
+#endif //WITH_MASSENTITY_DEBUG
 	return SNullWidget::NullWidget;
 }
 
