[Mass Debugger]
Fixed crash when clearing breakpoints from stale debug environemtns
Fixed missing scrollbars on breakpoints list
Fixed breakpoint list not showing the breakpoint list for the selected environment in some situations
Fixed several layout issues with entities list
Fixed several display issues when stopping PIE session with that environment selected in the debugger
Added option to auto-update entity fragment data view

diff --git a/Engine/Source/Editor/MassEntityDebugger/Private/MassDebuggerModel.cpp b/Engine/Source/Editor/MassEntityDebugger/Private/MassDebuggerModel.cpp
--- a/Engine/Source/Editor/MassEntityDebugger/Private/MassDebuggerModel.cpp
+++ b/Engine/Source/Editor/MassEntityDebugger/Private/MassDebuggerModel.cpp
@@ -412,6 +412,7 @@ PRAGMA_ENABLE_DEPRECATION_WARNINGS
 
 void FMassDebuggerModel::SetEnvironment(const TSharedPtr<FMassDebuggerEnvironment>& Item)
 {
+#if WITH_MASSENTITY_DEBUG
 	if (Item)
 	{
 		Environment = Item;
@@ -424,10 +425,12 @@ void FMassDebuggerModel::SetEnvironment(const TSharedPtr<FMassDebuggerEnvironmen
 	}
 
 	RefreshAll();
+#endif // WITH_MASSENTITY_DEBUG
 }
 
 void FMassDebuggerModel::RefreshAll()
 {
+#if WITH_MASSENTITY_DEBUG
 	if (Environment)
 	{
 		TMap<FMassArchetypeHandle, TSharedPtr<FMassDebuggerArchetypeData>> TransientArchetypesMap;
@@ -441,6 +444,7 @@ void FMassDebuggerModel::RefreshAll()
 
 		OnRefreshDelegate.Broadcast();
 	}
+#endif // WITH_MASSENTITY_DEBUG
 }
 
 void FMassDebuggerModel::SelectProcessor(TSharedPtr<FMassDebuggerProcessorData>& Processor)
@@ -797,6 +801,17 @@ TWeakPtr<SMassEntitiesView> FMassDebuggerModel::ShowEntitiesView(int32 Index)
 	return EntityViews[Index];
 }
 
+void FMassDebuggerModel::ResetEntitiesViews()
+{
+	for (int i = 0; i < MaxEntityViewCount; i++)
+	{
+		if (EntityViews[i].IsValid())
+		{
+			EntityViews[i].Pin()->ClearEntities();
+		}
+	}
+}
+
 void FMassDebuggerModel::ResetSelectedArchetypes()
 {
 	for (TSharedPtr<FMassDebuggerArchetypeData>& ArchetypeData : SelectedArchetypes)
diff --git a/Engine/Source/Editor/MassEntityDebugger/Private/MassDebuggerModel.h b/Engine/Source/Editor/MassEntityDebugger/Private/MassDebuggerModel.h
--- a/Engine/Source/Editor/MassEntityDebugger/Private/MassDebuggerModel.h
+++ b/Engine/Source/Editor/MassEntityDebugger/Private/MassDebuggerModel.h
@@ -209,6 +209,7 @@ struct FMassDebuggerModel
 	void ShowEntitiesView(int Index, TArray<FMassEntityHandle> EntitieHandles);
 	void ShowEntitiesView(int Index, FMassEntityQuery& Query);
 	void ShowEntitiesView(int Index, TConstArrayView<FMassEntityQuery*> InQueries);
+	void ResetEntitiesViews();
 	TWeakPtr<SMassEntitiesView> ShowEntitiesView(int32 Index);
 
 protected:
diff --git a/Engine/Source/Editor/MassEntityDebugger/Private/SMassBreakpointsView.cpp b/Engine/Source/Editor/MassEntityDebugger/Private/SMassBreakpointsView.cpp
--- a/Engine/Source/Editor/MassEntityDebugger/Private/SMassBreakpointsView.cpp
+++ b/Engine/Source/Editor/MassEntityDebugger/Private/SMassBreakpointsView.cpp
@@ -37,7 +37,7 @@ void SMassBreakpointsView::Construct(const FArguments& InArgs, TSharedRef<FMassD
 			.OnClicked(this, &SMassBreakpointsView::ClearBreakpointsClicked)
 		]
 		+ SVerticalBox::Slot()
-		.AutoHeight()
+		.FillHeight(1.f)
 		.Padding(5)
 		[
 			WriteBreakpointsListView.ToSharedRef()
@@ -52,11 +52,16 @@ void SMassBreakpointsView::Construct(const FArguments& InArgs, TSharedRef<FMassD
 #endif
 }
 
+void SMassBreakpointsView::OnRefresh()
+{
+	RefreshBreakpoints();
+}
+
 void SMassBreakpointsView::RefreshBreakpoints()
 {
 #if WITH_MASSENTITY_DEBUG
 	WriteBreakpoints.Reset();
-	if (DebuggerModel && DebuggerModel->Environment)
+	if (DebuggerModel && DebuggerModel->Environment && DebuggerModel->Environment->EntityManager.IsValid())
 	{
 		FMassDebugger::FEnvironment* Env = FMassDebugger::FindEnvironmentForEntityManager(DebuggerModel->Environment->EntityManager.Pin().ToSharedRef().Get());
 		TMultiMap<const UScriptStruct*, FMassEntityHandle>::TIterator FragmentIterator = Env->FragmentWriteBreakpoints.CreateIterator();
diff --git a/Engine/Source/Editor/MassEntityDebugger/Private/SMassBreakpointsView.h b/Engine/Source/Editor/MassEntityDebugger/Private/SMassBreakpointsView.h
--- a/Engine/Source/Editor/MassEntityDebugger/Private/SMassBreakpointsView.h
+++ b/Engine/Source/Editor/MassEntityDebugger/Private/SMassBreakpointsView.h
@@ -25,10 +25,7 @@ public:
 protected:
 
 	/** Unused pure virtual function */
-	virtual void OnRefresh() override
-	{
-
-	}
+	virtual void OnRefresh() override;
 
 	/** Unused pure virtual function */
 	virtual void OnProcessorsSelected(TConstArrayView<TSharedPtr<FMassDebuggerProcessorData>> SelectedProcessors, ESelectInfo::Type SelectInfo) override
diff --git a/Engine/Source/Editor/MassEntityDebugger/Private/SMassDebugger.cpp b/Engine/Source/Editor/MassEntityDebugger/Private/SMassDebugger.cpp
--- a/Engine/Source/Editor/MassEntityDebugger/Private/SMassDebugger.cpp
+++ b/Engine/Source/Editor/MassEntityDebugger/Private/SMassDebugger.cpp
@@ -266,6 +266,10 @@ void SMassDebugger::OnEntityManagerInitialized(const FMassEntityManager& EntityM
 	if (World && UE::Mass::Debugger::Private::IsSupportedWorldType(World->WorldType))
 	{
 		EnvironmentsList.Add(MakeShareable(new FMassDebuggerEnvironment(EntityManager.AsShared())));
+		if (EnvironmentComboBox.IsValid())
+		{
+			EnvironmentComboBox->RefreshOptions();
+		}
 	}
 }
 
@@ -325,6 +329,10 @@ void SMassDebugger::OnEntityManagerDeinitialized(const FMassEntityManager& Entit
 			}
 		}
 	}
+	if (EnvironmentComboBox.IsValid())
+	{
+		EnvironmentComboBox->RefreshOptions();
+	}
 }
 
 void SMassDebugger::HandleEnvironmentChanged(TSharedPtr<FMassDebuggerEnvironment> Item, ESelectInfo::Type SelectInfo)
@@ -414,6 +422,11 @@ bool SMassDebugger::CanRefreshData()
 
 void SMassDebugger::RefreshData()
 {
+	if (DebuggerModel->IsStale() && EnvironmentComboBox.IsValid() && EnvironmentComboLabel.IsValid())
+	{
+		EnvironmentComboBox->RefreshOptions();
+		EnvironmentComboLabel->SetText(LOCTEXT("PickEnvironment", "Pick Environment"));
+	}
 	DebuggerModel->RefreshAll();
 }
 
diff --git a/Engine/Source/Editor/MassEntityDebugger/Private/SMassEntitiesList.cpp b/Engine/Source/Editor/MassEntityDebugger/Private/SMassEntitiesList.cpp
--- a/Engine/Source/Editor/MassEntityDebugger/Private/SMassEntitiesList.cpp
+++ b/Engine/Source/Editor/MassEntityDebugger/Private/SMassEntitiesList.cpp
@@ -41,29 +41,16 @@ void SMassEntitiesList::Construct(const SMassEntitiesList::FArguments& InArgs, T
 		SNew(SBorder)
 		.Padding(5.0f)
 		[
-			SNew(SVerticalBox)
-			+ SVerticalBox::Slot()
-			.AutoHeight()
-			.Padding(5.0f)
-			.HAlign(HAlign_Left)
-			[
-				FragmentSelectBox.ToSharedRef()
-			]
-			+ SVerticalBox::Slot()
-			.VAlign(VAlign_Fill)
-			.FillHeight(1.f)
-			[
-				SAssignNew(TreeView, STreeView<EntitiesTableRowPtr>)
-				.SelectionMode(ESelectionMode::None)
-				.TreeItemsSource(&GridRows)
-				.OnGetChildren(this, &SMassEntitiesList::TreeView_OnGetChildren)
-				.OnGenerateRow(this, &SMassEntitiesList::TreeView_OnGenerateRow)
-				.HeaderRow
-				(
-					SAssignNew(TreeViewHeaderRow, SHeaderRow)
-					.Visibility(EVisibility::Visible)
-				)
-			]
+			SAssignNew(TreeView, STreeView<EntitiesTableRowPtr>)
+			.SelectionMode(ESelectionMode::None)
+			.TreeItemsSource(&GridRows)
+			.OnGetChildren(this, &SMassEntitiesList::TreeView_OnGetChildren)
+			.OnGenerateRow(this, &SMassEntitiesList::TreeView_OnGenerateRow)
+			.HeaderRow
+			(
+				SAssignNew(TreeViewHeaderRow, SHeaderRow)
+				.Visibility(EVisibility::Visible)
+			)
 		]
 	];
 
@@ -397,6 +384,15 @@ void SMassEntitiesList::RefreshEntityData()
 	RefreshFragmentData();
 }
 
+void SMassEntitiesList::Tick(const FGeometry& AllottedGeometry, const double InCurrentTime, const float InDeltaTime)
+{
+	if (bAutoUpdateEntityData)
+	{
+		RefreshEntityData();
+	}
+	SCompoundWidget::Tick(AllottedGeometry, InCurrentTime, InDeltaTime);
+}
+
 void SMassEntitiesList::UpdateTreeColumns()
 {
 #if WITH_MASSENTITY_DEBUG
diff --git a/Engine/Source/Editor/MassEntityDebugger/Private/SMassEntitiesList.h b/Engine/Source/Editor/MassEntityDebugger/Private/SMassEntitiesList.h
--- a/Engine/Source/Editor/MassEntityDebugger/Private/SMassEntitiesList.h
+++ b/Engine/Source/Editor/MassEntityDebugger/Private/SMassEntitiesList.h
@@ -28,10 +28,10 @@ public:
 	void Construct(const FArguments& InArgs, TSharedRef<FMassDebuggerModel> InDebuggerModel);
 	void SetEntities(const TArray<FMassEntityHandle>& InEntities);
 	void RefreshEntityData();
-protected:
-
 
+	virtual void Tick(const FGeometry& AllottedGeometry, const double InCurrentTime, const float InDeltaTime) override;
 
+protected:
 
 	struct FMassEntitiesListColumn
 	{
@@ -94,11 +94,12 @@ public:
 	TArray<FName> AvailableFragmentNames;
 	TArray<FName> SelectedFragmentNames;
 	TArray<const UScriptStruct*> SelectedFragmentTypes;
-	TSharedPtr<SWidget> FragmentSelectMenu;
 	FReply OnClearAllSelectedFragmentsClicked();
 
 	TSharedPtr<SBox> FragmentSelectBox;
 
+	bool bAutoUpdateEntityData = false;
+
 	void OnFragmentCheckStateChanged(ECheckBoxState NewState, FName FragmentName);
 	ECheckBoxState GetFragmentCheckState(FName FragmentName) const;
 	void CreateFragmentSelectDropdown();
diff --git a/Engine/Source/Editor/MassEntityDebugger/Private/SMassEntitiesView.cpp b/Engine/Source/Editor/MassEntityDebugger/Private/SMassEntitiesView.cpp
--- a/Engine/Source/Editor/MassEntityDebugger/Private/SMassEntitiesView.cpp
+++ b/Engine/Source/Editor/MassEntityDebugger/Private/SMassEntitiesView.cpp
@@ -2,6 +2,8 @@
 
 #include "SMassEntitiesView.h"
 #include "Widgets/Input/SButton.h"
+#include "Widgets/Input/SCheckBox.h"
+#include "Widgets/Input/SSlider.h"
 #include "Widgets/Layout/SScrollBox.h"
 #include "Widgets/Views/SListView.h"
 #include "Widgets/SBoxPanel.h"
@@ -27,6 +29,7 @@ void SMassEntitiesView::Construct(const FArguments& InArgs, TSharedRef<FMassDebu
 	[
 		SNew(SVerticalBox)
 		+ SVerticalBox::Slot()
+		.Padding(5)
 		.AutoHeight()
 		[
 			SNew(SHorizontalBox)
@@ -36,26 +39,51 @@ void SMassEntitiesView::Construct(const FArguments& InArgs, TSharedRef<FMassDebu
 				ListLabel.ToSharedRef()
 			]
 			+ SHorizontalBox::Slot()
+			.Padding(5)
+			.AutoWidth()
+			[
+				EntitiesList->FragmentSelectBox.ToSharedRef()
+			]
+			+ SHorizontalBox::Slot()
+			.Padding(5)
+			.AutoWidth()
+			[
+				SNew(SButton)
+				.VAlign(VAlign_Center)
+				.Text(LOCTEXT("ShowAllEntities", "ShowAllEntities"))
+				.OnClicked(FOnClicked::CreateSP(this, &SMassEntitiesView::ShowAllEntities))
+			]
+			+ SHorizontalBox::Slot()
+			.Padding(5)
 			.AutoWidth()
 			[
 				SNew(SButton)
+				.VAlign(VAlign_Center)
 				.Text(LOCTEXT("RefreshEntityList", "Refresh Entities"))
 				.OnClicked(FOnClicked::CreateSP(this, &SMassEntitiesView::RefreshEntityList))
 			]
 			+ SHorizontalBox::Slot()
+			.Padding(5)
 			.AutoWidth()
 			[
 				SNew(SButton)
+				.VAlign(VAlign_Center)
 				.Text(LOCTEXT("RefreshEntityData", "Refresh Entity Data"))
 				.OnClicked(FOnClicked::CreateSP(this, &SMassEntitiesView::RefreshEntityData))
 			]
 			+ SHorizontalBox::Slot()
+			.Padding(5)
 			.AutoWidth()
 			[
-				SNew(SButton)
-				.Text(LOCTEXT("ShowAllEntities", "ShowAllEntities"))
-				.OnClicked(FOnClicked::CreateSP(this, &SMassEntitiesView::ShowAllEntities))
+
+				SAssignNew(AutoUpdateCheckbox, SCheckBox)
+				.OnCheckStateChanged(this, &SMassEntitiesView::OnAutoUpdateChanged)
+				[
+					SNew(STextBlock)
+					.Text(LOCTEXT("AutoUpdateEntityData", "Auto Update Entity Data"))
+				]
 			]
+
 		]
 		+ SVerticalBox::Slot()
 		.FillHeight(1.f)
@@ -75,6 +103,14 @@ void SMassEntitiesView::Construct(const FArguments& InArgs, TSharedRef<FMassDebu
 #endif
 }
 
+void SMassEntitiesView::OnAutoUpdateChanged(ECheckBoxState NewState)
+{
+	if (EntitiesList.IsValid())
+	{
+		EntitiesList->bAutoUpdateEntityData = (NewState == ECheckBoxState::Checked);
+	}
+}
+
 void SMassEntitiesView::ShowEntities(FMassArchetypeHandle InArchetypeHandle)
 {
 #if WITH_MASSENTITY_DEBUG
@@ -124,13 +160,28 @@ FReply SMassEntitiesView::ShowAllEntities()
 	return FReply::Handled();
 }
 
+void SMassEntitiesView::ClearEntities()
+{
+#if WITH_MASSENTITY_DEBUG
+	TempEntityList.Reset();
+	ShowEntitiesFrom = EShowEntitiesFrom::List;
+	RefreshEntityList();
+#endif
+}
+
 FReply SMassEntitiesView::RefreshEntityList()
 {
 #if WITH_MASSENTITY_DEBUG
-	if (!EntitiesList.IsValid() || !DebuggerModel.IsValid() || !DebuggerModel->Environment.IsValid() || !DebuggerModel->Environment->EntityManager.IsValid())
-	{ 
+	if (!DebuggerModel.IsValid() || !DebuggerModel->Environment.IsValid() || !DebuggerModel->Environment->EntityManager.IsValid())
+	{
+		TempEntityList.Reset();
+		if (EntitiesList.IsValid())
+		{
+			EntitiesList->SetEntities(TempEntityList);
+		}
 		return FReply::Handled();
 	}
+
 	const FMassEntityManager& EntityManager = *DebuggerModel->Environment->EntityManager.Pin().Get();
 
 	switch (ShowEntitiesFrom)
@@ -188,5 +239,12 @@ FReply SMassEntitiesView::RefreshEntityData()
 	return FReply::Handled();
 }
 
+void SMassEntitiesView::OnRefresh()
+{
+#if WITH_MASSENTITY_DEBUG
+	RefreshEntityList();
+#endif
+}
+
 #undef LOCTEXT_NAMESPACE
 
diff --git a/Engine/Source/Editor/MassEntityDebugger/Private/SMassEntitiesView.h b/Engine/Source/Editor/MassEntityDebugger/Private/SMassEntitiesView.h
--- a/Engine/Source/Editor/MassEntityDebugger/Private/SMassEntitiesView.h
+++ b/Engine/Source/Editor/MassEntityDebugger/Private/SMassEntitiesView.h
@@ -8,6 +8,7 @@
 #include "MassEntityQuery.h"
 
 class SMassEntitiesList;
+class SCheckBox;
 
 class SMassEntitiesView : public SMassDebuggerViewBase
 {
@@ -23,6 +24,7 @@ public:
 	void ShowEntities(const TArray<FMassEntityHandle>& InEntities);
 	void ShowEntities(FMassEntityQuery& InQuery);
 	void ShowEntities(TConstArrayView<FMassEntityQuery*> InQueries);
+	void ClearEntities();
 
 protected:
 	enum class EShowEntitiesFrom
@@ -37,11 +39,11 @@ protected:
 	FReply RefreshEntityList();
 	FReply RefreshEntityData();
 
-	/** Unused pure virtual function */
-	virtual void OnRefresh() override
-	{
+	void OnAutoUpdateChanged(ECheckBoxState NewState);
+	float GetUpdateInterval() const;
+	void OnUpdateIntervalChanged(float NewValue);
 
-	}
+	virtual void OnRefresh() override;
 
 	/** Unused pure virtual function */
 	virtual void OnProcessorsSelected(TConstArrayView<TSharedPtr<FMassDebuggerProcessorData>> SelectedProcessors, ESelectInfo::Type SelectInfo) override
@@ -63,4 +65,5 @@ protected:
 	uint32 EntitiesViewIndex;
 	FMassEntityQuery Query;
 	TConstArrayView<FMassEntityQuery*> Queries;
+	TSharedPtr<SCheckBox> AutoUpdateCheckbox;
 };
