[Mass Debugger]
Added support for shared fragments and const shared fragments to the entity data display.
Added scrollbar to the fragment selection drop-down

diff --git a/Engine/Source/Editor/MassEntityDebugger/Private/SMassEntitiesList.cpp b/Engine/Source/Editor/MassEntityDebugger/Private/SMassEntitiesList.cpp
--- a/Engine/Source/Editor/MassEntityDebugger/Private/SMassEntitiesList.cpp
+++ b/Engine/Source/Editor/MassEntityDebugger/Private/SMassEntitiesList.cpp
@@ -3,6 +3,7 @@
 #include "SMassEntitiesList.h"
 #include "MassDebuggerModel.h"
 #include "MassEntityTypes.h"
+#include "MassEntityView.h"
 #include "SMassBitSet.h"
 #include "Widgets/Text/SRichTextBlock.h"
 #include "Styling/AppStyle.h"
@@ -181,6 +182,20 @@ void SMassEntitiesList::SetEntities(const TArray<FMassEntityHandle>& InEntities)
 				{
 					AvailableFragmentNames.AddUnique(FragmentType->GetFName());
 				});
+
+				const FMassArchetypeSharedFragmentValues& SharedFragments = FMassDebugger::GetSharedFragmentValues(EntityManager, Entity);
+				const TArray<FConstSharedStruct>& ConstSharedStructs = SharedFragments.GetConstSharedFragments();
+				const TArray<FSharedStruct>& SharedStructs = SharedFragments.GetSharedFragments();
+
+				for (const FSharedStruct& SharedStruct : SharedStructs)
+				{
+					AvailableFragmentNames.AddUnique(SharedStruct.GetScriptStruct()->GetFName());
+				}
+
+				for (const FConstSharedStruct& ConstSharedStruct : ConstSharedStructs)
+				{
+					AvailableFragmentNames.AddUnique(ConstSharedStruct.GetScriptStruct()->GetFName());
+				}
 			}
 		}
 	}
@@ -232,10 +247,22 @@ void SMassEntitiesList::RefreshFragmentData()
 
 				for (FGridRow::FFragmentInfo& Info : Row.FragmentInfo)
 				{
-					if (Info.StructData.IsValid())
+					if (UE::Mass::IsA<FMassFragment>(Info.StructType))
 					{
 						FMassDebugger::GetFragmentData(EntityManager, Info.StructType, Row.Entity, Info.StructData);
 					}
+					else if (UE::Mass::IsA<FMassSharedFragment>(Info.StructType))
+					{
+						FMassDebugger::GetSharedFragmentData(EntityManager, Info.StructType, Row.Entity, Info.StructData);
+					}
+					else if (UE::Mass::IsA<FMassConstSharedFragment>(Info.StructType))
+					{
+						FMassDebugger::GetConstSharedFragmentData(EntityManager, Info.StructType, Row.Entity, Info.StructData);
+					}
+					else
+					{
+						ensureMsgf(false, TEXT("Invalid entity data type(%s)"), *Info.StructType->GetDisplayNameText().ToString());
+					}
 				}
 			}
 		}
@@ -340,10 +367,24 @@ void SMassEntitiesList::CreateFragmentSelectDropdown()
 			.Text(LOCTEXT("ClearAll", "Clear All"))
 			.OnClicked(FOnClicked::CreateSP(this, &SMassEntitiesList::OnClearAllSelectedFragmentsClicked))
 	];
+	
+	TSharedPtr<SScrollBox> FragmentScrollBox = SNew(SScrollBox);
+	DropdownContent->AddSlot()
+	.FillHeight(1.)
+	.Padding(5.0f)
+	[
+		FragmentScrollBox.ToSharedRef()
+	];
+
+	TSharedPtr<SVerticalBox> FragmentList = SNew(SVerticalBox);
+	FragmentScrollBox->AddSlot()
+	[
+		FragmentList.ToSharedRef()
+	];
 
 	for (const FName& FragmentName : AvailableFragmentNames)
 	{
-		DropdownContent->AddSlot()
+		FragmentList->AddSlot()
 		.AutoHeight()
 		.Padding(5.0f)
 		[
@@ -525,16 +566,41 @@ void SMassEntitiesList::SEntitiesTableRow::Construct(const FArguments& InArgs, c
 
 			for (int i = 0; i < FragmentDisplayCount; i++)
 			{
-				if (Row.FragmentInfo[i].StructType != OwnerList.SelectedFragmentTypes[i])
+				const UScriptStruct* StructType = OwnerList.SelectedFragmentTypes[i];
+				if (StructType && Row.FragmentInfo[i].StructType != StructType)
 				{
-					Row.FragmentInfo[i].StructType = OwnerList.SelectedFragmentTypes[i];
-					Row.FragmentInfo[i].StructData = FMassDebugger::GetFragmentData(EntityManager, OwnerList.SelectedFragmentTypes[i], Row.Entity);
-					if (Row.FragmentInfo[i].StructData.IsValid())
+					Row.FragmentInfo[i].StructType = StructType;
+					if (UE::Mass::IsA<FMassFragment>(StructType))
+					{
+						Row.FragmentInfo[i].StructData = FMassDebugger::GetFragmentData(EntityManager, OwnerList.SelectedFragmentTypes[i], Row.Entity);
+						if (Row.FragmentInfo[i].StructData.IsValid())
+						{
+							Row.FragmentInfo[i].PropertyRowGenerator = PropertyEditorModule.CreatePropertyRowGenerator(GeneratorArgs);
+							Row.FragmentInfo[i].PropertyRowGenerator->SetStructure(Row.FragmentInfo[i].StructData);
+						}
+					}
+					else if (UE::Mass::IsA<FMassSharedFragment>(StructType))
+					{
+						Row.FragmentInfo[i].StructData = FMassDebugger::GetSharedFragmentData(EntityManager, OwnerList.SelectedFragmentTypes[i], Row.Entity);
+						if (Row.FragmentInfo[i].StructData.IsValid())
+						{
+							Row.FragmentInfo[i].PropertyRowGenerator = PropertyEditorModule.CreatePropertyRowGenerator(GeneratorArgs);
+							Row.FragmentInfo[i].PropertyRowGenerator->SetStructure(Row.FragmentInfo[i].StructData);
+						}
+					}
+					else if (UE::Mass::IsA<FMassConstSharedFragment>(StructType))
+					{
+						Row.FragmentInfo[i].StructData = FMassDebugger::GetConstSharedFragmentData(EntityManager, OwnerList.SelectedFragmentTypes[i], Row.Entity);
+						if (Row.FragmentInfo[i].StructData.IsValid())
+						{
+							Row.FragmentInfo[i].PropertyRowGenerator = PropertyEditorModule.CreatePropertyRowGenerator(GeneratorArgs);
+							Row.FragmentInfo[i].PropertyRowGenerator->SetStructure(Row.FragmentInfo[i].StructData);
+						}
+					}
+					else
 					{
-						Row.FragmentInfo[i].PropertyRowGenerator = PropertyEditorModule.CreatePropertyRowGenerator(GeneratorArgs);
-						Row.FragmentInfo[i].PropertyRowGenerator->SetStructure(Row.FragmentInfo[i].StructData);
+						ensureMsgf(false, TEXT("Invalid entity data type(%s)"), *StructType->GetDisplayNameText().ToString());
 					}
-					
 				}
 			}
 		}
diff --git a/Engine/Source/Runtime/MassEntity/Private/MassDebugger.cpp b/Engine/Source/Runtime/MassEntity/Private/MassDebugger.cpp
--- a/Engine/Source/Runtime/MassEntity/Private/MassDebugger.cpp
+++ b/Engine/Source/Runtime/MassEntity/Private/MassDebugger.cpp
@@ -1213,6 +1213,84 @@ bool FMassDebugger::GetFragmentData(const FMassEntityManager& EntityManager, con
 	return false;
 }
 
+const FMassArchetypeSharedFragmentValues& FMassDebugger::GetSharedFragmentValues(const FMassEntityManager& EntityManager, FMassEntityHandle Entity)
+{
+	TSharedPtr<FMassArchetypeData> Archetype = EntityManager.DebugGetEntityStorageInterface().GetArchetypeAsShared(Entity.Index);
+	if (Archetype.IsValid())
+	{
+		return Archetype->GetSharedFragmentValues(Entity);
+	}
+
+	static FMassArchetypeSharedFragmentValues Dummy;
+	return Dummy;
+}
+
+TSharedPtr<FStructOnScope> FMassDebugger::GetSharedFragmentData(const FMassEntityManager& EntityManager, const UScriptStruct* FragmentType, FMassEntityHandle Entity)
+{
+	TSharedPtr<FStructOnScope> StructOnScope = MakeShared<FStructOnScope>(FragmentType);
+	if (GetSharedFragmentData(EntityManager, FragmentType, Entity, StructOnScope))
+	{
+		return StructOnScope;
+	}
+	return nullptr;
+}
+
+bool FMassDebugger::GetSharedFragmentData(const FMassEntityManager& EntityManager, const UScriptStruct* FragmentType, FMassEntityHandle Entity, TSharedPtr<FStructOnScope>& OutStructData)
+{
+	TSharedPtr<FMassArchetypeData> Archetype = EntityManager.DebugGetEntityStorageInterface().GetArchetypeAsShared(Entity.Index);
+	if (Archetype.IsValid())
+	{
+		const FSharedStruct* SharedFragment = Archetype->GetSharedFragmentValues(Entity).GetSharedFragments().FindByPredicate(FStructTypeEqualOperator(FragmentType));
+		void* FragmentData = (SharedFragment != nullptr) ? SharedFragment->GetMemory() : nullptr;
+
+		if (FragmentData)
+		{
+			const UStruct* AsUStruct = FragmentType;
+			if (OutStructData->GetStruct() != AsUStruct)
+			{
+				OutStructData->Initialize(FragmentType);
+			}
+
+			CastChecked<UScriptStruct>(OutStructData->GetStruct())->CopyScriptStruct(OutStructData->GetStructMemory(), FragmentData);
+			return true;
+		}
+	}
+	return false;
+}
+
+TSharedPtr<FStructOnScope> FMassDebugger::GetConstSharedFragmentData(const FMassEntityManager& EntityManager, const UScriptStruct* FragmentType, FMassEntityHandle Entity)
+{
+	TSharedPtr<FStructOnScope> StructOnScope = MakeShared<FStructOnScope>(FragmentType);
+	if (GetConstSharedFragmentData(EntityManager, FragmentType, Entity, StructOnScope))
+	{
+		return StructOnScope;
+	}
+	return nullptr;
+}
+
+bool FMassDebugger::GetConstSharedFragmentData(const FMassEntityManager& EntityManager, const UScriptStruct* FragmentType, FMassEntityHandle Entity, TSharedPtr<FStructOnScope>& OutStructData)
+{
+	TSharedPtr<FMassArchetypeData> Archetype = EntityManager.DebugGetEntityStorageInterface().GetArchetypeAsShared(Entity.Index);
+	if (Archetype.IsValid())
+	{
+		const FConstSharedStruct* SharedFragment = Archetype->GetSharedFragmentValues(Entity).GetConstSharedFragments().FindByPredicate(FStructTypeEqualOperator(FragmentType));
+		const void* FragmentData = (SharedFragment != nullptr) ? SharedFragment->GetMemory() : nullptr;
+
+		if (FragmentData)
+		{
+			const UStruct* AsUStruct = FragmentType;
+			if (OutStructData->GetStruct() != AsUStruct)
+			{
+				OutStructData->Initialize(FragmentType);
+			}
+
+			CastChecked<UScriptStruct>(OutStructData->GetStruct())->CopyScriptStruct(OutStructData->GetStructMemory(), FragmentData);
+			return true;
+		}
+	}
+	return false;
+}
+
 void FMassDebugger::UpdateHasBreakpoint()
 {
 	bHasBreakpoint = false;
diff --git a/Engine/Source/Runtime/MassEntity/Public/MassDebugger.h b/Engine/Source/Runtime/MassEntity/Public/MassDebugger.h
--- a/Engine/Source/Runtime/MassEntity/Public/MassDebugger.h
+++ b/Engine/Source/Runtime/MassEntity/Public/MassDebugger.h
@@ -331,10 +331,34 @@ struct MASSENTITY_API FMassDebugger
 	static TSharedPtr<FStructOnScope> GetFragmentData(const FMassEntityManager& EntityManager, const UScriptStruct* FragmentType, FMassEntityHandle Entity);
 
 	/**
-	 * Finds the fragment data of the specified type in the entity data. Returns nullptr if not found.
+	 * Finds the fragment data of the specified type in the entity data. Returns false if not found.
 	 */
 	static bool GetFragmentData(const FMassEntityManager& EntityManager, const UScriptStruct* FragmentType, FMassEntityHandle Entity, TSharedPtr<FStructOnScope>& OutStructData);
 
+	/**
+	 * Get the shared fragment value container for this entity
+	 */
+	static const FMassArchetypeSharedFragmentValues& GetSharedFragmentValues(const FMassEntityManager& EntityManager, FMassEntityHandle Entity);
+	
+	/**
+	 * Finds the shared fragment data of the specified type in the entity data. Returns nullptr if not found.
+	 */
+	static TSharedPtr<FStructOnScope> GetSharedFragmentData(const FMassEntityManager& EntityManager, const UScriptStruct* FragmentType, FMassEntityHandle Entity);
+
+	/**
+	 * Finds the shared fragment data of the specified type in the entity data. Returns nullptr if not found.
+	 */
+	static bool GetSharedFragmentData(const FMassEntityManager& EntityManager, const UScriptStruct* FragmentType, FMassEntityHandle Entity, TSharedPtr<FStructOnScope>& OutStructData);
+	
+	/**
+	 * Finds the const shared fragment data of the specified type in the entity data. Returns nullptr if not found.
+	 */
+	static TSharedPtr<FStructOnScope> GetConstSharedFragmentData(const FMassEntityManager& EntityManager, const UScriptStruct* FragmentType, FMassEntityHandle Entity);
+
+	/**
+	 * Finds the const shared fragment data of the specified type in the entity data. Returns nullptr if not found.
+	 */
+	static bool GetConstSharedFragmentData(const FMassEntityManager& EntityManager, const UScriptStruct* FragmentType, FMassEntityHandle Entity, TSharedPtr<FStructOnScope>& OutStructData);
 
 	/**
 	 * Clears all breakpoints in all environments.
