Updated Mass Debugger's environment picker to use the debug display name of the world instead of its path, making it easier to identify client, server and editor worlds

diff --git a/Engine/Source/Editor/MassEntityDebugger/Private/MassDebuggerModel.cpp b/Engine/Source/Editor/MassEntityDebugger/Private/MassDebuggerModel.cpp
--- a/Engine/Source/Editor/MassEntityDebugger/Private/MassDebuggerModel.cpp
+++ b/Engine/Source/Editor/MassEntityDebugger/Private/MassDebuggerModel.cpp
@@ -87,7 +87,7 @@ FString FMassDebuggerEnvironment::GetDisplayName() const
 #endif // WITH_MASSENTITY_DEBUG
 	
 	const UWorld* WorldPtr = World.Get();
-	DisplayName += WorldPtr ? WorldPtr->GetPathName() : TEXT("No World");
+	DisplayName += WorldPtr ? WorldPtr->GetDebugDisplayName() : TEXT("No World");
 	return DisplayName;
 }
 
