[Mass Debugger]
Fixed crash bug when debug environment becomes invalid with breakpoints set.

diff --git a/Engine/Source/Editor/MassEntityDebugger/Private/SMassBreakpointsView.cpp b/Engine/Source/Editor/MassEntityDebugger/Private/SMassBreakpointsView.cpp
--- a/Engine/Source/Editor/MassEntityDebugger/Private/SMassBreakpointsView.cpp
+++ b/Engine/Source/Editor/MassEntityDebugger/Private/SMassBreakpointsView.cpp
@@ -64,6 +64,11 @@ void SMassBreakpointsView::RefreshBreakpoints()
 	if (DebuggerModel && DebuggerModel->Environment && DebuggerModel->Environment->EntityManager.IsValid())
 	{
 		FMassDebugger::FEnvironment* Env = FMassDebugger::FindEnvironmentForEntityManager(DebuggerModel->Environment->EntityManager.Pin().ToSharedRef().Get());
+		if (Env == nullptr)
+		{
+			return;
+		}
+
 		TMultiMap<const UScriptStruct*, FMassEntityHandle>::TIterator FragmentIterator = Env->FragmentWriteBreakpoints.CreateIterator();
 		for (; FragmentIterator; ++FragmentIterator)
 		{
