[MassDebugger]
Fixed crash bug with stale environments

diff --git a/Engine/Source/Editor/MassEntityDebugger/Private/SMassEntitiesList.cpp b/Engine/Source/Editor/MassEntityDebugger/Private/SMassEntitiesList.cpp
--- a/Engine/Source/Editor/MassEntityDebugger/Private/SMassEntitiesList.cpp
+++ b/Engine/Source/Editor/MassEntityDebugger/Private/SMassEntitiesList.cpp
@@ -233,9 +233,7 @@ TSharedPtr<IPropertyHandle> FindPropertyHandle(const FProperty* Property, TArray
 void SMassEntitiesList::RefreshFragmentData()
 {
 #if WITH_MASSENTITY_DEBUG
-	if (DebuggerModel.IsValid()
-		&& DebuggerModel->Environment.IsValid()
-		&& DebuggerModel->Environment->EntityManager.IsValid())
+	if (DebuggerModel.IsValid() && !DebuggerModel->IsStale())
 	{
 		const FMassEntityManager& EntityManager = *(DebuggerModel->Environment->EntityManager.Pin());
 
@@ -551,9 +549,7 @@ void SMassEntitiesList::SEntitiesTableRow::Construct(const FArguments& InArgs, c
 	if (TableRowPtr.IsValid())
 	{
 		SMassEntitiesList& OwnerList = *Row.EntitiesList.Pin();
-		if (OwnerList.DebuggerModel.IsValid()
-			&& OwnerList.DebuggerModel->Environment.IsValid()
-			&& OwnerList.DebuggerModel->Environment->EntityManager.IsValid())
+		if (OwnerList.DebuggerModel.IsValid() && !OwnerList.DebuggerModel->IsStale())
 		{
 			const FMassEntityManager& EntityManager = *(OwnerList.DebuggerModel->Environment->EntityManager.Pin());
 			FPropertyRowGeneratorArgs GeneratorArgs;
