[MassDebugger]
Fix missing guards for cooked editor build

diff --git a/Engine/Source/Editor/MassEntityDebugger/Private/MassDebuggerModel.cpp b/Engine/Source/Editor/MassEntityDebugger/Private/MassDebuggerModel.cpp
--- a/Engine/Source/Editor/MassEntityDebugger/Private/MassDebuggerModel.cpp
+++ b/Engine/Source/Editor/MassEntityDebugger/Private/MassDebuggerModel.cpp
@@ -795,10 +795,14 @@ void FMassDebuggerModel::MarkAsStale()
 
 bool FMassDebuggerModel::IsStale() const 
 { 
+#if WITH_MASSENTITY_DEBUG
 	return Environment.IsValid() == false
 		|| Environment->EntityManager.IsValid() == false
 		|| FMassDebugger::IsEntityManagerInitialized(*Environment->EntityManager.Pin()) == false
 		|| (Environment->NeedsValidWorld() == true && Environment->IsWorldValid() == false);
+#else
+	return true;
+#endif //WITH_MASSENTITY_DEBUG
 }
 
 const TSharedPtr<FMassDebuggerProcessorData>& FMassDebuggerModel::GetProcessorDataChecked(const UMassProcessor& Processor) const
diff --git a/Engine/Source/Editor/MassEntityDebugger/Private/SMassFragmentsView.cpp b/Engine/Source/Editor/MassEntityDebugger/Private/SMassFragmentsView.cpp
--- a/Engine/Source/Editor/MassEntityDebugger/Private/SMassFragmentsView.cpp
+++ b/Engine/Source/Editor/MassEntityDebugger/Private/SMassFragmentsView.cpp
@@ -186,8 +186,6 @@ void SFragmentsTableRow::Construct(const FArguments& InArgs, const TSharedRef<ST
 		.ShowSelection(false)
 		, InOwnerTableView);
 
-#else // WITH_MASSENTITY_DEBUG
-	SMultiColumnTableRow<FFragmentsTableRowPtr>::Construct(FArguments(), InOwnerTableView);
 #endif //WITH_MASSENTITY_DEBUG
 }
 
