[MassDebugger]
Added handle sort to Entities tab.

diff --git a/Engine/Source/Editor/MassEntityDebugger/Private/SMassEntitiesList.cpp b/Engine/Source/Editor/MassEntityDebugger/Private/SMassEntitiesList.cpp
--- a/Engine/Source/Editor/MassEntityDebugger/Private/SMassEntitiesList.cpp
+++ b/Engine/Source/Editor/MassEntityDebugger/Private/SMassEntitiesList.cpp
@@ -30,6 +30,12 @@
 //----------------------------------------------------------------------//
 // SMassEntitiesList
 //----------------------------------------------------------------------//
+
+namespace UE::MassDebugger::EntitiesList::Private
+{
+	const FLazyName ColumnHandle = TEXT("EntityHandle");
+}
+
 void SMassEntitiesList::Construct(const SMassEntitiesList::FArguments& InArgs, TSharedRef<FMassDebuggerModel> InDebuggerModel)
 {
 #if WITH_MASSENTITY_DEBUG
@@ -435,7 +441,6 @@ void SMassEntitiesList::Tick(const FGeometry& AllottedGeometry, const double InC
 void SMassEntitiesList::UpdateTreeColumns()
 {
 #if WITH_MASSENTITY_DEBUG
-	const FName HandleName = FName(TEXT("EntityHandle"));
 
 	const TIndirectArray<SHeaderRow::FColumn>& TreeColumns = TreeViewHeaderRow->GetColumns();
 	bool Changed = false;
@@ -444,16 +449,17 @@ void SMassEntitiesList::UpdateTreeColumns()
 	{
 		SHeaderRow::FColumn::FArguments ColumnArgs;
 		ColumnArgs
-		.ColumnId(HandleName)
+		.ColumnId(UE::MassDebugger::EntitiesList::Private::ColumnHandle.Resolve())
 		.DefaultLabel(LOCTEXT("MassEntityHandle", "Entity Handle"))
 		.ToolTipText(LOCTEXT("MassEntityHandle", "Entity Handle"))
 		.HAlignHeader(HAlign_Left)
 		.VAlignHeader(VAlign_Center)
 		.HAlignCell(HAlign_Fill)
 		.VAlignCell(VAlign_Fill)
-		.InitialSortMode(EColumnSortMode::Ascending)
-		.FillWidth(50.f)
-		.FillSized(60)
+		.SortMode(this, &SMassEntitiesList::GetColumnSortMode, UE::MassDebugger::EntitiesList::Private::ColumnHandle.Resolve())
+		.OnSort(this, &SMassEntitiesList::OnColumnSortModeChanged)
+		.FillWidth(100.f)
+		.FillSized(100.f)
 		.HeaderContent()
 		[
 			SNew(SBox)
@@ -524,6 +530,42 @@ void SMassEntitiesList::UpdateTreeColumns()
 #endif //WITH_MASSENTITY_DEBUG
 }
 
+EColumnSortMode::Type SMassEntitiesList::GetColumnSortMode(FName ColumnId) const
+{
+	if (ColumnId != UE::MassDebugger::EntitiesList::Private::ColumnHandle)
+	{
+		return EColumnSortMode::None;
+	}
+
+	return bSortAscending ? EColumnSortMode::Ascending : EColumnSortMode::Descending;
+}
+
+void SMassEntitiesList::OnColumnSortModeChanged(const EColumnSortPriority::Type SortPriority, const FName& ColumnId, const EColumnSortMode::Type InSortMode)
+{
+	if (ColumnId != UE::MassDebugger::EntitiesList::Private::ColumnHandle)
+	{
+		return;
+	}
+
+	if (SortPriority == EColumnSortPriority::Primary)
+	{
+		bSortAscending = InSortMode == EColumnSortMode::Ascending;
+	}
+
+	auto SortHandle = [bAscending = bSortAscending](const TSharedPtr<FGridRow>& A, const TSharedPtr<FGridRow>& B)
+	{
+		const int32 Compare = B->Entity.Index - A->Entity.Index;
+		return bAscending ? Compare > 0 : Compare <= 0;
+	};
+	
+	GridRows.StableSort(SortHandle);
+
+	if (TreeView.IsValid())
+	{
+		TreeView->RequestListRefresh();
+	}
+}
+
 void SMassEntitiesList::TreeView_OnGetChildren(SMassEntitiesList::EntitiesTableRowPtr InParent,
 	TArray<SMassEntitiesList::EntitiesTableRowPtr>& OutChildren)
 {
diff --git a/Engine/Source/Editor/MassEntityDebugger/Private/SMassEntitiesList.h b/Engine/Source/Editor/MassEntityDebugger/Private/SMassEntitiesList.h
--- a/Engine/Source/Editor/MassEntityDebugger/Private/SMassEntitiesList.h
+++ b/Engine/Source/Editor/MassEntityDebugger/Private/SMassEntitiesList.h
@@ -82,6 +82,17 @@ public:
 		TSharedRef<SWidget> GenerateDataWidget(const FProperty* Property, FGridRow::FFragmentInfo& Info);
 	};
 
+	TSharedPtr<SBox> GetFragmentSelectBox()
+	{
+		return FragmentSelectBox;
+	}
+
+	void AutoUpdateEntityData(bool bEnable)
+	{
+		bAutoUpdateEntityData = bEnable;
+	}
+
+private:
 	TArray<EntitiesTableRowPtr> GridRows;
 	TArray<FMassEntitiesListColumn> Columns;
 	TMap<FName, int32> ColumnIndexByID;
@@ -105,6 +116,9 @@ public:
 	void CreateFragmentSelectDropdown();
 	TSharedPtr<FMassDebuggerModel> DebuggerModel;
 	void UpdateTreeColumns();
+	EColumnSortMode::Type GetColumnSortMode(FName ColumnId) const;
+	void OnColumnSortModeChanged(const EColumnSortPriority::Type SortPriority, const FName& ColumnId, const EColumnSortMode::Type InSortMode);
+	bool bSortAscending = true;
 
 	void TreeView_OnGetChildren(SMassEntitiesList::EntitiesTableRowPtr InParent, TArray<SMassEntitiesList::EntitiesTableRowPtr>& OutChildren);
 	TSharedRef<ITableRow> TreeView_OnGenerateRow(EntitiesTableRowPtr RowPtr, const TSharedRef<STableViewBase>& OwnerTable);
diff --git a/Engine/Source/Editor/MassEntityDebugger/Private/SMassEntitiesView.cpp b/Engine/Source/Editor/MassEntityDebugger/Private/SMassEntitiesView.cpp
--- a/Engine/Source/Editor/MassEntityDebugger/Private/SMassEntitiesView.cpp
+++ b/Engine/Source/Editor/MassEntityDebugger/Private/SMassEntitiesView.cpp
@@ -42,7 +42,7 @@ void SMassEntitiesView::Construct(const FArguments& InArgs, TSharedRef<FMassDebu
 			.Padding(5)
 			.AutoWidth()
 			[
-				EntitiesList->FragmentSelectBox.ToSharedRef()
+				EntitiesList->GetFragmentSelectBox().ToSharedRef()
 			]
 			+ SHorizontalBox::Slot()
 			.Padding(5)
@@ -107,7 +107,7 @@ void SMassEntitiesView::OnAutoUpdateChanged(ECheckBoxState NewState)
 {
 	if (EntitiesList.IsValid())
 	{
-		EntitiesList->bAutoUpdateEntityData = (NewState == ECheckBoxState::Checked);
+		EntitiesList->AutoUpdateEntityData(NewState == ECheckBoxState::Checked);
 	}
 }
 
