[MassDebugger]
Added query editor UI with load/save and copy/paste support for authored queries

diff --git a/Engine/Source/Editor/MassEntityDebugger/MassEntityDebugger.Build.cs b/Engine/Source/Editor/MassEntityDebugger/MassEntityDebugger.Build.cs
--- a/Engine/Source/Editor/MassEntityDebugger/MassEntityDebugger.Build.cs
+++ b/Engine/Source/Editor/MassEntityDebugger/MassEntityDebugger.Build.cs
@@ -23,6 +23,14 @@ namespace UnrealBuildTool.Rules
 				}
 			);
 
+			PrivateDependencyModuleNames.AddRange(
+				new string[] {
+					"ApplicationCore",
+					"Json",
+					"ToolWidgets"
+				}
+			);
+
 			if (Target.bBuildEditor == true)
 			{
 				PublicDependencyModuleNames.AddRange(
diff --git a/Engine/Source/Editor/MassEntityDebugger/Private/MassDebuggerModel.cpp b/Engine/Source/Editor/MassEntityDebugger/Private/MassDebuggerModel.cpp
--- a/Engine/Source/Editor/MassEntityDebugger/Private/MassDebuggerModel.cpp
+++ b/Engine/Source/Editor/MassEntityDebugger/Private/MassDebuggerModel.cpp
@@ -96,6 +96,19 @@ TSharedPtr<const FMassEntityManager> FMassDebuggerEnvironment::GetEntityManager(
 	return EntityManager.Pin();
 }
 
+TSharedPtr<FMassEntityManager> FMassDebuggerEnvironment::GetMutableEntityManager() const
+{
+	if (TSharedPtr<const FMassEntityManager> SharedConstEntityManager = GetEntityManager())
+	{
+		FMassDebugger::FEnvironment* Environment = FMassDebugger::FindEnvironmentForEntityManager(*SharedConstEntityManager);
+		if (Environment)
+		{
+			return Environment->MutableEntityManager.Pin();
+		}
+	}
+	return nullptr;
+}
+
 //----------------------------------------------------------------------//
 // FMassDebuggerQueryData
 //----------------------------------------------------------------------//
@@ -456,6 +469,7 @@ void FMassDebuggerModel::RefreshAll()
 		CacheProcessorsData(TransientArchetypesMap, CompositeProcessors);
 		CacheProcessingGraphs(CompositeProcessors);
 		CacheFragmentData();
+		CacheTagData();
 
 		ClearArchetypeSelection();
 
@@ -625,14 +639,17 @@ void FMassDebuggerModel::CacheProcessingGraphs(TConstArrayView<TNotNull<const UM
 	}
 }
 
-void FMassDebuggerModel::CacheFragmentData()
+void FMassDebuggerModel::CacheFragmentData(
+	TArray<TSharedPtr<FMassDebuggerFragmentData>>& OutData,
+	const TConstArrayView<TWeakObjectPtr<const UScriptStruct>>& FragmentTypes
+)
 {
-	CachedFragmentData.Reset();
+	OutData.Reset();
 
 #if WITH_MASSENTITY_DEBUG
 	TMap<const UScriptStruct*, TSharedPtr<FMassDebuggerFragmentData>> TempFragmentDataMap;
 
-	for (const TWeakObjectPtr<const UScriptStruct>& WeakStruct : FMassFragmentBitSet::DebugGetAllStructTypes())
+	for (const TWeakObjectPtr<const UScriptStruct>& WeakStruct : FragmentTypes)
 	{
 		if (WeakStruct.IsValid())
 		{
@@ -663,16 +680,16 @@ void FMassDebuggerModel::CacheFragmentData()
 		});
 	}
 
-	CachedFragmentData.Reserve(TempFragmentDataMap.Num());
+	OutData.Reserve(TempFragmentDataMap.Num());
 	for (auto const& [FragmentStruct, FragmentData] : TempFragmentDataMap)
 	{
-		CachedFragmentData.Add(FragmentData);
+		OutData.Add(FragmentData);
 	}
 
 	if (IsCurrentEnvironmentValid())
 	{
 		TWeakPtr<const FMassEntityManager> EntityManager = Environment->GetEntityManager().ToWeakPtr();
-		for (TSharedPtr<FMassDebuggerFragmentData>& FragmentData : CachedFragmentData)
+		for (TSharedPtr<FMassDebuggerFragmentData>& FragmentData : OutData)
 		{
 			FragmentData->EntityManager = EntityManager;
 		}
@@ -682,10 +699,20 @@ void FMassDebuggerModel::CacheFragmentData()
 	{
 		return A->Name.CompareTo(B->Name) <= 0;
 	};
-	CachedFragmentData.Sort(SortPredicate);
+	OutData.Sort(SortPredicate);
 #endif
 }
 
+void FMassDebuggerModel::CacheFragmentData()
+{
+	CacheFragmentData(CachedFragmentData, FMassFragmentBitSet::DebugGetAllStructTypes());
+}
+
+void FMassDebuggerModel::CacheTagData()
+{
+	CacheFragmentData(CachedTagData, FMassTagBitSet::DebugGetAllStructTypes());
+}
+
 void FMassDebuggerModel::CacheArchetypesData(TMap<FMassArchetypeHandle, TSharedPtr<FMassDebuggerArchetypeData>>& OutTransientArchetypesMap)
 {
 	CachedAllArchetypes.Reset();
diff --git a/Engine/Source/Editor/MassEntityDebugger/Private/MassDebuggerModel.h b/Engine/Source/Editor/MassEntityDebugger/Private/MassDebuggerModel.h
--- a/Engine/Source/Editor/MassEntityDebugger/Private/MassDebuggerModel.h
+++ b/Engine/Source/Editor/MassEntityDebugger/Private/MassDebuggerModel.h
@@ -27,7 +27,7 @@ enum class EMassDebuggerSelectionMode : uint8
 	Archetype,
 	// @todo future:
 	// Fragment
-	MAX
+	MAX	
 };
 
 enum class EMassDebuggerProcessorSelection : uint8
@@ -176,6 +176,7 @@ struct FMassDebuggerEnvironment
 
 	FString GetDisplayName() const;
 	TSharedPtr<const FMassEntityManager> GetEntityManager() const;
+	TSharedPtr<FMassEntityManager> GetMutableEntityManager() const;
 	bool IsWorldValid() const { return World.IsValid(); }
 	bool NeedsValidWorld() const { return bNeedsValidWorld; }
 	
@@ -215,7 +216,12 @@ struct FMassDebuggerModel
 	void CacheProcessorsData(const TMap<FMassArchetypeHandle, TSharedPtr<FMassDebuggerArchetypeData>>& InTransientArchetypesMap, TArray<TNotNull<const UMassCompositeProcessor*>>& OutCompositeProcessors);
 	void CacheProcessingGraphs(TConstArrayView<TNotNull<const UMassCompositeProcessor*>> InCompositeProcessors);
 
+	void CacheFragmentData(
+		TArray<TSharedPtr<FMassDebuggerFragmentData>>& OutData,
+		const TConstArrayView<TWeakObjectPtr<const UScriptStruct>>& FragmentTypes
+	);
 	void CacheFragmentData();
+	void CacheTagData();
 
 	float MinDistanceToSelectedArchetypes(const TSharedPtr<FMassDebuggerArchetypeData>& InArchetypeData) const;
 
@@ -265,6 +271,7 @@ public:
 		TArray<TSharedPtr<FMassDebuggerProcessorData>> Container;
 	};
 
+	TArray<TSharedPtr<FMassDebuggerFragmentData>> CachedTagData;
 	TArray<TSharedPtr<FMassDebuggerFragmentData>> CachedFragmentData;
 	TArray<TSharedPtr<FProcessorCollection>> CachedProcessorCollections;
 	TArray<TSharedPtr<FMassDebuggerProcessorData>> SelectedProcessors;
@@ -288,6 +295,7 @@ protected:
 	TArray<TSharedPtr<FMassDebuggerProcessorData>> AllCachedProcessors;
 	TArray<TWeakPtr<SMassEntitiesView>> EntityViews;
 	FName SelectedFragmentName;
+
 public:
 	UE_DEPRECATED(5.6, "CachedProcessors property is now deprecated. Use CachedProcessorCollections instead. ")
 	TArray<TSharedPtr<FMassDebuggerProcessorData>> CachedProcessors;
diff --git a/Engine/Source/Editor/MassEntityDebugger/Private/SMassDebugger.cpp b/Engine/Source/Editor/MassEntityDebugger/Private/SMassDebugger.cpp
--- a/Engine/Source/Editor/MassEntityDebugger/Private/SMassDebugger.cpp
+++ b/Engine/Source/Editor/MassEntityDebugger/Private/SMassDebugger.cpp
@@ -7,6 +7,7 @@
 #include "SMassArchetypesView.h"
 #include "SMassEntitiesView.h"
 #include "SMassFragmentsView.h"
+#include "SMassQueryEditorView.h"
 #include "MassDebuggerModel.h"
 #include "Engine/Engine.h"
 #include "CoreGlobals.h"
@@ -40,6 +41,7 @@ namespace UE::Mass::Debugger::Private
 	const FLazyName FragmentsTabId("Fragments");
 	const FLazyName ArchetypesTabId("Archetypes");
 	const FLazyName EntitiesTabId("Entities");
+	const FLazyName QueryEditorTabId("QueryEditor");
 
 	bool IsSupportedWorldType(const EWorldType::Type WorldType)
 	{
@@ -113,6 +115,10 @@ void SMassDebugger::Construct(const FArguments& InArgs, const TSharedRef<SDockTa
 		.SetDisplayName(LOCTEXT("EntitiesTabTitle", "Entities"))
 		.SetGroup(AppMenuGroup);
 
+	TabManager->RegisterTabSpawner(QueryEditorTabId, FOnSpawnTab::CreateRaw(this, &SMassDebugger::SpawnQueryEditorTab))
+		.SetDisplayName(LOCTEXT("QueryEditorTabTitle", "Query Editor"))
+		.SetGroup(AppMenuGroup);
+
 	// Default Layout
 	TSharedRef<FTabManager::FLayout> Layout = CreateDefaultLayout();
 
@@ -141,7 +147,7 @@ void SMassDebugger::Construct(const FArguments& InArgs, const TSharedRef<SDockTa
 TSharedRef<FTabManager::FLayout> SMassDebugger::CreateDefaultLayout()
 {
 	using namespace UE::Mass::Debugger::Private;
-	TSharedRef<FTabManager::FLayout> Layout = FTabManager::NewLayout("MassDebuggerLayout_v1.2")
+	TSharedRef<FTabManager::FLayout> Layout = FTabManager::NewLayout("MassDebuggerLayout_v1.3")
 		->AddArea
 		(
 			FTabManager::NewPrimaryArea()
@@ -166,6 +172,7 @@ TSharedRef<FTabManager::FLayout> SMassDebugger::CreateDefaultLayout()
 						->AddTab(ProcessorsTabId, ETabState::OpenedTab)
 						->AddTab(ProcessingGraphTabId, ETabState::OpenedTab)
 						->AddTab(FragmentsTabId, ETabState::OpenedTab)
+						->AddTab(QueryEditorTabId, ETabState::OpenedTab)
 						->SetForegroundTab(ProcessorsTabId.Resolve())
 					)
 					->Split
@@ -254,6 +261,7 @@ TSharedRef<SWidget> SMassDebugger::GenerateWindowMenu()
 	FUIAction FragmentsUIAction(FExecuteAction::CreateSP(this, &SMassDebugger::ShowFragmentsView));
 	FUIAction BreakpointsUIAction(FExecuteAction::CreateSP(this, &SMassDebugger::ShowBreakpointsView));
 	FUIAction EntitesUIAction(FExecuteAction::CreateSP(this, &SMassDebugger::ShowEntitesView));
+	FUIAction QueryEditorUIAction(FExecuteAction::CreateSP(this, &SMassDebugger::ShowQueryEditorView));
 	FUIAction ResetLayoutUIAction(FExecuteAction::CreateSP(this, &SMassDebugger::ResetLayout));
 
 	MenuBuilder.BeginSection(TEXT("Tabs"));
@@ -263,6 +271,7 @@ TSharedRef<SWidget> SMassDebugger::GenerateWindowMenu()
 	MenuBuilder.AddMenuEntry(LOCTEXT("FragmentsTabLabel", "Fragments"), LOCTEXT("FragmentsTabTooltip", "Show Fragments"), FSlateIcon(), FragmentsUIAction);
 	MenuBuilder.AddMenuEntry(LOCTEXT("BreakpointsTabLabel", "Breakpoints"), LOCTEXT("BreakpointsTabTooltip", "Show Breakpoints Tab"), FSlateIcon(), BreakpointsUIAction);
 	MenuBuilder.AddMenuEntry(LOCTEXT("EntitiesTabLabel", "Entities"), LOCTEXT("EntitiesTabTooltip", "Show Entities Tab"), FSlateIcon(), EntitesUIAction);
+	MenuBuilder.AddMenuEntry(LOCTEXT("QueryEditorTabLabel", "Query Editor"), LOCTEXT("QueryEditoTabTooltip", "Show Query Editor Tab"), FSlateIcon(), QueryEditorUIAction);
 	MenuBuilder.AddSeparator();
 	MenuBuilder.AddMenuEntry(LOCTEXT("ResetLayoutLabel", "Reset Layout"), LOCTEXT("ResetLayoutTooltip", "Reset the Mass Debugger Layout"), FSlateIcon(), ResetLayoutUIAction);
 
@@ -435,6 +444,17 @@ TSharedRef<SDockTab> SMassDebugger::SpawnEntitiesTab(const FSpawnTabArgs& Args)
 	return MajorTab;
 }
 
+TSharedRef<SDockTab> SMassDebugger::SpawnQueryEditorTab(const FSpawnTabArgs& Args)
+{
+	const TSharedRef<SDockTab> MajorTab = SNew(SDockTab)
+		.TabRole(ETabRole::PanelTab);
+
+	TSharedPtr<SWidget> TabContent = SNew(UE::MassDebugger::SQueryEditorView, DebuggerModel);
+	MajorTab->SetContent(TabContent.ToSharedRef());
+
+	return MajorTab;
+}
+
 bool SMassDebugger::CanRefreshData()
 {
 	return DebuggerModel->HasEnvironmentSelected();
@@ -480,6 +500,11 @@ void SMassDebugger::ShowEntitesView() const
 	TabManager->TryInvokeTab(UE::Mass::Debugger::Private::EntitiesTabId.Resolve());
 }
 
+void SMassDebugger::ShowQueryEditorView() const
+{
+	TabManager->TryInvokeTab(UE::Mass::Debugger::Private::QueryEditorTabId.Resolve());
+}
+
 void SMassDebugger::ResetLayout()
 {
 	TSharedPtr<SWindow> Window = FSlateApplication::Get().FindWidgetWindow(AsShared());
@@ -495,7 +520,6 @@ void SMassDebugger::ShowSelectedView() const
 	
 }
 
-
 void SMassDebugger::BindDelegates()
 {
 #if WITH_MASSENTITY_DEBUG
diff --git a/Engine/Source/Editor/MassEntityDebugger/Private/SMassDebugger.h b/Engine/Source/Editor/MassEntityDebugger/Private/SMassDebugger.h
--- a/Engine/Source/Editor/MassEntityDebugger/Private/SMassDebugger.h
+++ b/Engine/Source/Editor/MassEntityDebugger/Private/SMassDebugger.h
@@ -64,6 +64,7 @@ public:
 	void ShowEntitesView() const;
 	void ResetLayout();
 	void ShowSelectedView() const;
+	void ShowQueryEditorView() const;
 
 protected:
 	TSharedRef<SDockTab> SpawnToolbar(const FSpawnTabArgs& Args);
@@ -73,6 +74,7 @@ protected:
 	TSharedRef<SDockTab> SpawnArchetypesTab(const FSpawnTabArgs& Args);
 	TSharedRef<SDockTab> SpawnBreakpointsTab(const FSpawnTabArgs& Args);
 	TSharedRef<SDockTab> SpawnEntitiesTab(const FSpawnTabArgs& Args);
+	TSharedRef<SDockTab> SpawnQueryEditorTab(const FSpawnTabArgs& Args);
 
 	TSharedRef<FTabManager::FLayout> CreateDefaultLayout();
 	TSharedRef<SWidget> GenerateWindowMenu();
diff --git a/Engine/Source/Editor/MassEntityDebugger/Private/SMassFragmentsView.cpp b/Engine/Source/Editor/MassEntityDebugger/Private/SMassFragmentsView.cpp
--- a/Engine/Source/Editor/MassEntityDebugger/Private/SMassFragmentsView.cpp
+++ b/Engine/Source/Editor/MassEntityDebugger/Private/SMassFragmentsView.cpp
@@ -143,8 +143,8 @@ void SFragmentsView::RefreshFragmentList()
 	{
 		auto SortName = [bAscending = bSortAscending](const TSharedPtr<FMassDebuggerFragmentData>& A, const TSharedPtr<FMassDebuggerFragmentData>& B)
 		{
-			const int32 Compare = A->Name.CompareTo(B->Name) <= 0;
-			return bAscending ? Compare > 0 : Compare <= 0;
+			const int32 Compare = A->Name.CompareTo(B->Name);
+			return bAscending ? Compare < 0 : Compare >= 0;
 		};
 		FragmentListItemsSource.StableSort(SortName);
 	}
diff --git a/Engine/Source/Editor/MassEntityDebugger/Private/SMassQueryEditor.cpp b/Engine/Source/Editor/MassEntityDebugger/Private/SMassQueryEditor.cpp
new file mode 100644
--- /dev/null
+++ b/Engine/Source/Editor/MassEntityDebugger/Private/SMassQueryEditor.cpp
@@ -0,0 +1,495 @@
+// Copyright Epic Games, Inc. All Rights Reserved.
+
+#include "SMassQueryEditor.h"
+#include "SSearchableComboBox.h"
+#include "SMassQueryEditorView.h"
+#include "Widgets/Layout/SBox.h"
+#include "Widgets/Views/SListView.h"
+#include "Widgets/Input/SButton.h"
+#include "Widgets/Input/SComboBox.h"
+#include "Widgets/Input/SEditableTextBox.h"
+#include "Widgets/Text/STextBlock.h"
+#include "Internationalization/Internationalization.h"
+
+#define LOCTEXT_NAMESPACE "SMassRequirementsEditor"
+
+namespace UE::MassDebugger
+{
+
+	void SQueryEditor::Construct(const FArguments& InArgs)
+	{
+		DebuggerModel = InArgs._DebuggerModel;
+		OnRequirementsChanged = InArgs._OnRequirementsChanged;
+
+		TWeakPtr<SQueryEditor> WeakThisPtr = StaticCastWeakPtr<SQueryEditor>(AsWeak());
+
+		if (TSharedPtr<FMassDebuggerModel> Model = DebuggerModel.Pin())
+		{
+			Model->CacheFragmentData();
+			Model->CacheTagData();
+
+			for (const TSharedPtr<FMassDebuggerFragmentData>& FragData : Model->CachedFragmentData)
+			{
+				if (FragData->Fragment.IsValid())
+				{
+					int32 Index = AvailableFragmentNames.Add(MakeShared<FString>(FragData->Fragment->GetName()));
+					FragmentNamesToTypes.Add(*AvailableFragmentNames[Index], FragData->Fragment);
+				}
+			}
+			for (const TSharedPtr<FMassDebuggerFragmentData>& TagData : Model->CachedTagData)
+			{
+				if (TagData->Fragment.IsValid())
+				{
+					int32 Index = AvailableTagNames.Add(MakeShared<FString>(TagData->Fragment->GetName()));
+					TagNamesToTypes.Add(*AvailableTagNames[Index], TagData->Fragment);
+				}
+			}
+		}
+
+		const UEnum* FragmentAccessEnum = StaticEnum<EMassFragmentAccess>();
+		for (int32 i = 0; i < FragmentAccessEnum->NumEnums() - 1; ++i)
+		{
+			int64 AccessValue = FragmentAccessEnum->GetValueByIndex(i);
+			if (AccessValue != INDEX_NONE)
+			{
+				FragmentAccessNames.Add(MakeShared<FName>(FragmentAccessEnum->GetNameByIndex(i)));
+				FragmentAccessMap.Add(FragmentAccessEnum->GetNameByIndex(i), static_cast<EMassFragmentAccess>(AccessValue));
+			}
+		}
+
+		const UEnum* FragmentPresenceEnum = StaticEnum<EMassFragmentPresence>();
+		for (int32 i = 0; i < FragmentPresenceEnum->NumEnums() - 1; ++i)
+		{
+			int64 PresenceValue = FragmentPresenceEnum->GetValueByIndex(i);
+			if (PresenceValue != INDEX_NONE)
+			{
+				FragmentPresenceNames.Add(MakeShared<FName>(FragmentPresenceEnum->GetNameByIndex(i)));
+				FragmentPresenceMap.Add(FragmentPresenceEnum->GetNameByIndex(i), static_cast<EMassFragmentPresence>(PresenceValue));
+			}
+		}
+
+		Load();
+
+		ChildSlot
+		[
+			SNew(SVerticalBox)
+			+ SVerticalBox::Slot()
+			.AutoHeight()
+			.Padding(2.0f)
+			[
+				SNew(SEditableTextBox)
+				.Text_Lambda([WeakThisPtr]()
+				{
+					if (TSharedPtr<SQueryEditor> SharedThis = WeakThisPtr.Pin())
+					{
+						return SharedThis->Query.IsValid()
+							? FText::FromString(SharedThis->Query->Name)
+							: FText::GetEmpty();
+					}
+
+					return FText::GetEmpty();
+				})
+				.OnTextCommitted(this, &SQueryEditor::HandleNameTextCommitted)
+			]
+			+ SVerticalBox::Slot()
+			.AutoHeight()
+			.Padding(2.0f)
+			[
+				SNew(SHorizontalBox)
+				+ SHorizontalBox::Slot()
+				.AutoWidth()
+				[
+					SNew(STextBlock)
+					.Text(LOCTEXT("FragmentsLabel", "Fragment Requirements"))
+				]
+				+ SHorizontalBox::Slot()
+				.FillWidth(1.0f)
+				[
+					SNew(SButton)
+					.Text(LOCTEXT("AddFragBtn", "Add Fragment"))
+					.OnClicked(this, &SQueryEditor::OnAddFragment)
+				]
+			]
+			+ SVerticalBox::Slot()
+			.FillHeight(0.5f)
+			.Padding(2.0f)
+			[
+				SAssignNew(FragmentListView, SListView<TSharedPtr<FFragmentEntry>>)
+				.ListItemsSource(&FragmentEntries)
+				.OnGenerateRow(this, &SQueryEditor::OnGenerateFragmentRow)
+			]
+			+ SVerticalBox::Slot()
+			.AutoHeight()
+			.Padding(.0f)
+			[
+				SNew(SHorizontalBox)
+				+ SHorizontalBox::Slot()
+				.AutoWidth()
+				[
+					SNew(STextBlock)
+					.Text(LOCTEXT("TagsLabel", "Tag Requirements"))
+				]
+				+ SHorizontalBox::Slot()
+				.FillWidth(1.0f)
+				[
+					SNew(SButton)
+					.Text(LOCTEXT("AddTagBtn", "Add Tag"))
+					.OnClicked(this, &SQueryEditor::OnAddTag)
+				]
+			]
+			+ SVerticalBox::Slot()
+			.FillHeight(0.5f)
+			.Padding(2.0f)
+			[
+				SAssignNew(TagListView, SListView<TSharedPtr<FFragmentEntry>>)
+				.ListItemsSource(&TagEntries)
+				.OnGenerateRow(this, &SQueryEditor::OnGenerateTagRow)
+			]
+		];
+	}
+
+	void SQueryEditor::SetQuery(TSharedPtr<FEditableQuery>& InQuery)
+	{
+		Query = InQuery;
+		Load();
+	}
+
+	void SQueryEditor::ApplyChanges()
+	{
+		if (Query.IsValid())
+		{
+			Query->FragmentRequirements = FragmentEntries;
+			Query->TagRequirements = TagEntries;
+		}
+		else
+		{
+			Load();
+		}
+	}
+
+	void SQueryEditor::Load()
+	{
+		if (Query.IsValid())
+		{
+			FragmentEntries = Query->FragmentRequirements;
+			TagEntries = Query->TagRequirements;
+		}
+		else
+		{
+			FragmentEntries.Reset();
+			TagEntries.Empty();
+		}
+
+		if (FragmentListView.IsValid())
+		{
+			FragmentListView->RequestListRefresh();
+		}
+		if (TagListView.IsValid())
+		{
+			TagListView->RequestListRefresh();
+		}
+	}
+
+	void SQueryEditor::HandleNameTextCommitted(const FText& InText, ETextCommit::Type InCommitType)
+	{
+		if (Query.IsValid())
+		{
+			Query->Name = InText.ToString();
+			NotifyRequirementsChanged();
+		}
+	}
+
+	TSharedRef<ITableRow> SQueryEditor::OnGenerateFragmentRow(TSharedPtr<FFragmentEntry> InEntry, const TSharedRef<STableViewBase>& OwnerTable)
+	{
+		TSharedPtr<SQueryEditor> SharedThis = StaticCastSharedPtr<SQueryEditor>(AsShared().ToSharedPtr());
+		return SNew(SFragmentSelectorRow, OwnerTable, InEntry, SharedThis.ToSharedRef(), false);
+	}
+
+	void SQueryEditor::SFragmentSelectorRow::Construct(const FArguments& InArgs, const TSharedRef<STableViewBase>& InOwnerTable, const TSharedPtr<FFragmentEntry>& InEntry, TSharedRef<SQueryEditor> InOwnerRequirementsEditor, bool bInShowTags)
+	{
+		OwnerRequirementsEditor = InOwnerRequirementsEditor;
+		ConstructInternal(InArgs, InOwnerTable);
+
+		TArray<TSharedPtr<FString>>* FragmentNamesPtr;
+		TSharedPtr<SWidget> AccessModeSelector;
+		TWeakPtr<SQueryEditor::SFragmentSelectorRow> WeakThisPtr = StaticCastWeakPtr<SQueryEditor::SFragmentSelectorRow>(AsWeak());
+
+		bShowTags = bInShowTags;
+		if (bShowTags)
+		{
+			FragmentNamesPtr = &InOwnerRequirementsEditor->AvailableTagNames;
+			AccessModeSelector = SNullWidget::NullWidget.ToSharedPtr();
+		}
+		else
+		{
+			FragmentNamesPtr = &InOwnerRequirementsEditor->AvailableFragmentNames;
+
+			AccessModeSelector = SNew(SComboBox<TSharedPtr<FName>>)
+			.OptionsSource(&InOwnerRequirementsEditor->FragmentAccessNames)
+			.OnGenerateWidget_Lambda([](TSharedPtr<FName> InName)
+			{
+				return SNew(STextBlock).Text(FText::FromName(*InName));
+			})
+			.OnSelectionChanged_Lambda([WeakThisPtr, InEntry](TSharedPtr<FName> NewName, ESelectInfo::Type)
+			{
+				TSharedPtr<SQueryEditor::SFragmentSelectorRow> SharedThis = WeakThisPtr.Pin();
+				if (SharedThis.IsValid() && InEntry.IsValid() && NewName.IsValid() && SharedThis->OwnerRequirementsEditor.IsValid())
+				{
+					SQueryEditor& Editor = *SharedThis->OwnerRequirementsEditor.Pin();
+					InEntry->AccessMode = Editor.FragmentAccessMap[*NewName];
+					Editor.NotifyRequirementsChanged();
+				}
+			})
+			[
+				SNew(STextBlock).Text_Lambda([InEntry]()
+				{
+					switch (InEntry->AccessMode)
+					{
+					case EMassFragmentAccess::ReadOnly:
+						return LOCTEXT("ReadOnly", "ReadOnly");
+					case EMassFragmentAccess::ReadWrite:
+						return LOCTEXT("ReadWrite", "ReadWrite");
+					default:
+						return LOCTEXT("None", "<None>");
+					}
+				})
+			];
+		}
+
+		ChildSlot
+		[
+			SNew(SHorizontalBox)
+			+ SHorizontalBox::Slot().AutoWidth().Padding(2)
+			[
+				SAssignNew(FragmentsComboBox, SSearchableComboBox)
+				.OptionsSource(FragmentNamesPtr)
+				.OnGenerateWidget_Lambda([](const TSharedPtr<FString>& InItem)
+				{
+					return SNew(STextBlock)
+						.Text(InItem ? FText::FromString(*InItem) : FText::GetEmpty());
+				})
+				.OnSelectionChanged_Lambda([WeakThisPtr, InEntry](const TSharedPtr<FString> Value, ESelectInfo::Type)
+				{
+					TSharedPtr<SQueryEditor::SFragmentSelectorRow> SharedThis = WeakThisPtr.Pin();
+					if (SharedThis.IsValid() && InEntry.IsValid() && Value.IsValid() && SharedThis->OwnerRequirementsEditor.IsValid())
+					{
+						SQueryEditor& Editor = *SharedThis->OwnerRequirementsEditor.Pin();
+						if (SharedThis->bShowTags)
+						{
+							InEntry->StructType = Editor.TagNamesToTypes[*Value];
+						}
+						else
+						{
+							InEntry->StructType = Editor.FragmentNamesToTypes[*Value];
+						}
+						Editor.NotifyRequirementsChanged();
+					}
+				})
+				.Content()
+				[
+					SNew(STextBlock)
+					.Text_Lambda([InEntry]() -> FText
+					{
+						if (InEntry.IsValid() && InEntry->StructType.IsValid())
+						{
+							return InEntry->StructType->GetDisplayNameText();
+						}
+						return LOCTEXT("None", "<None>");
+					})
+				]
+
+
+			]
+			+ SHorizontalBox::Slot().AutoWidth().Padding(2)
+			[
+				AccessModeSelector.ToSharedRef()
+			]
+			+ SHorizontalBox::Slot().AutoWidth().Padding(2)
+			[
+				SNew(SComboBox<TSharedPtr<FName>>)
+				.OptionsSource(&InOwnerRequirementsEditor->FragmentPresenceNames)
+				.OnSelectionChanged_Lambda([WeakThisPtr, InEntry](TSharedPtr<FName> NewName, ESelectInfo::Type)
+				{
+					TSharedPtr<SQueryEditor::SFragmentSelectorRow> SharedThis = WeakThisPtr.Pin();
+					if (SharedThis.IsValid() && InEntry.IsValid() && NewName.IsValid() && SharedThis->OwnerRequirementsEditor.IsValid())
+					{
+						SQueryEditor& Editor = *SharedThis->OwnerRequirementsEditor.Pin();
+						InEntry->Presence = Editor.FragmentPresenceMap[*NewName];
+						Editor.NotifyRequirementsChanged();
+					}
+				})
+				.OnGenerateWidget_Lambda([](TSharedPtr<FName> InName)
+				{
+					return SNew(STextBlock).Text(FText::FromName(*InName));
+				})
+				[
+					SNew(STextBlock).Text_Lambda([InEntry]()
+					{
+						if (!InEntry.IsValid()) return LOCTEXT("None", "<None>");
+						switch (InEntry->Presence)
+						{
+						case EMassFragmentPresence::All:
+							return LOCTEXT("All", "All");
+						case EMassFragmentPresence::Any:
+							return LOCTEXT("Any", "Any");
+						case EMassFragmentPresence::Optional:
+							return LOCTEXT("Optional", "Optional");
+						default:
+							return LOCTEXT("None", "None");
+						}
+					})
+				]
+			]
+			+ SHorizontalBox::Slot().AutoWidth().Padding(2)
+			[
+				SNew(SButton)
+				.Text(LOCTEXT("Remove", "Remove"))
+				.OnClicked(
+					&InOwnerRequirementsEditor.Get(),
+					bShowTags ? &SQueryEditor::OnRemoveTag : &SQueryEditor::OnRemoveFragment,
+					InEntry)
+			]
+		];
+	}
+
+	TSharedRef<ITableRow> SQueryEditor::OnGenerateTagRow(TSharedPtr<FFragmentEntry> InEntry, const TSharedRef<STableViewBase>& OwnerTable)
+	{
+		TSharedPtr<SQueryEditor> SharedThis = StaticCastSharedPtr<SQueryEditor>(AsShared().ToSharedPtr());
+		return SNew(SFragmentSelectorRow, OwnerTable, InEntry, SharedThis.ToSharedRef(), true);
+	}
+
+	FReply SQueryEditor::OnAddFragment()
+	{
+		TSharedPtr<FFragmentEntry> NewFragmentEntry = MakeShared<FFragmentEntry>();
+		if (AvailableFragmentNames.Num() > 0)
+		{
+			NewFragmentEntry->StructType = (FragmentNamesToTypes[*AvailableFragmentNames[0]]);
+		}
+		FragmentEntries.Add(NewFragmentEntry);
+		FragmentListView->RequestListRefresh();
+		NotifyRequirementsChanged();
+		return FReply::Handled();
+	}
+
+	FReply SQueryEditor::OnRemoveFragment(TSharedPtr<FFragmentEntry> InEntry)
+	{
+		FragmentEntries.Remove(InEntry);
+		FragmentListView->RequestListRefresh();
+		NotifyRequirementsChanged();
+		return FReply::Handled();
+	}
+
+	FReply SQueryEditor::OnAddTag()
+	{
+		TSharedPtr<FFragmentEntry> NewTagEntry = MakeShared<FFragmentEntry>();
+		if (AvailableTagNames.Num() > 0)
+		{
+			NewTagEntry->StructType = (TagNamesToTypes[*AvailableTagNames[0]]);
+		}
+		TagEntries.Add(NewTagEntry);
+		TagListView->RequestListRefresh();
+		NotifyRequirementsChanged();
+		return FReply::Handled();
+	}
+
+	FReply SQueryEditor::OnRemoveTag(TSharedPtr<FFragmentEntry> InEntry)
+	{
+		TagEntries.Remove(InEntry);
+		TagListView->RequestListRefresh();
+		NotifyRequirementsChanged();
+		return FReply::Handled();
+	}
+
+	void SQueryEditor::SFragmentSelectorRow::OnFragmentTypeChanged(TSharedPtr<TWeakObjectPtr<const UScriptStruct>> NewType, ESelectInfo::Type, TSharedPtr<FFragmentEntry> Entry)
+	{
+		if (Entry.IsValid() && NewType.IsValid() && NewType->Get() && OwnerRequirementsEditor.IsValid())
+		{
+			SQueryEditor& Editor = *OwnerRequirementsEditor.Pin();
+			Entry->StructType = NewType->Get();
+			Editor.NotifyRequirementsChanged();
+		}
+	}
+
+	void SQueryEditor::SFragmentSelectorRow::OnFragmentAccessChanged(EMassFragmentAccess NewAccess, ESelectInfo::Type, TSharedPtr<FFragmentEntry> Entry)
+	{
+		if (Entry.IsValid() && OwnerRequirementsEditor.IsValid())
+		{
+			SQueryEditor& Editor = *OwnerRequirementsEditor.Pin();
+			Entry->AccessMode = NewAccess;
+			Editor.NotifyRequirementsChanged();
+		}
+	}
+
+	void SQueryEditor::SFragmentSelectorRow::OnFragmentPresenceChanged(EMassFragmentPresence  NewPresence, ESelectInfo::Type, TSharedPtr<FFragmentEntry> Entry)
+	{
+		if (Entry.IsValid() && OwnerRequirementsEditor.IsValid())
+		{
+			SQueryEditor& Editor = *OwnerRequirementsEditor.Pin();
+			Entry->Presence = NewPresence;
+			Editor.NotifyRequirementsChanged();
+		}
+	}
+
+	void SQueryEditor::SFragmentSelectorRow::OnTagTypeChanged(TSharedPtr<TWeakObjectPtr<const UScriptStruct>> NewType, ESelectInfo::Type, TSharedPtr<FFragmentEntry> Entry)
+	{
+		if (Entry.IsValid() && NewType.IsValid() && NewType->Get() && OwnerRequirementsEditor.IsValid())
+		{
+			SQueryEditor& Editor = *OwnerRequirementsEditor.Pin();
+			Entry->StructType = NewType->Get();
+			Editor.NotifyRequirementsChanged();
+		}
+	}
+
+	void SQueryEditor::SFragmentSelectorRow::OnTagPresenceChanged(EMassFragmentPresence NewPresence, ESelectInfo::Type, TSharedPtr<FFragmentEntry> Entry)
+	{
+		if (Entry.IsValid() && OwnerRequirementsEditor.IsValid())
+		{
+			SQueryEditor& Editor = *OwnerRequirementsEditor.Pin();
+			Entry->Presence = NewPresence;
+			Editor.NotifyRequirementsChanged();
+		}
+	}
+
+	FText SQueryEditor::SFragmentSelectorRow::GetFragmentTypeText(TSharedPtr<FFragmentEntry> Entry) const
+	{
+		return Entry.IsValid() && Entry->StructType.IsValid()
+			? FText::FromName(Entry->StructType->GetFName())
+			: LOCTEXT("None", "<None>");
+	}
+
+	FText SQueryEditor::SFragmentSelectorRow::GetTagTypeText(TSharedPtr<FFragmentEntry> Entry) const
+	{
+		return Entry.IsValid() && Entry->StructType.IsValid()
+			? FText::FromName(Entry->StructType->GetFName())
+			: LOCTEXT("None", "<None>");
+	}
+
+	FText SQueryEditor::SFragmentSelectorRow::GetTagPresenceText(TSharedPtr<FFragmentEntry> Entry) const
+	{
+		if (!Entry.IsValid()) return LOCTEXT("NonePres", "<None>");
+		switch (Entry->Presence)
+		{
+		case EMassFragmentPresence::All:
+			return LOCTEXT("All", "All");
+		case EMassFragmentPresence::Any:
+			return LOCTEXT("Any", "Any");
+		case EMassFragmentPresence::Optional:
+			return LOCTEXT("Optional", "Optional");
+		default:
+			return LOCTEXT("None", "None");
+		}
+	}
+
+	void SQueryEditor::NotifyRequirementsChanged()
+	{
+		ApplyChanges();
+		
+		if (OnRequirementsChanged.IsBound())
+		{
+			OnRequirementsChanged.Execute(Query);
+		}
+	}
+
+} // namespace UE::MassDebugger
+
+#undef LOCTEXT_NAMESPACE
diff --git a/Engine/Source/Editor/MassEntityDebugger/Private/SMassQueryEditor.h b/Engine/Source/Editor/MassEntityDebugger/Private/SMassQueryEditor.h
new file mode 100644
--- /dev/null
+++ b/Engine/Source/Editor/MassEntityDebugger/Private/SMassQueryEditor.h
@@ -0,0 +1,100 @@
+// Copyright Epic Games, Inc. All Rights Reserved.
+
+#pragma once
+
+#include "CoreMinimal.h"
+#include "Widgets/SCompoundWidget.h"
+#include "MassRequirements.h"
+#include "MassDebuggerModel.h"
+#include "Widgets/Views/SListView.h"
+#include "Widgets/Views/STableRow.h"
+#include "Widgets/Input/SEditableTextBox.h"
+
+class SSearchableComboBox;
+
+namespace UE::MassDebugger
+{
+	struct FEditableQuery;
+	struct FFragmentEntry;
+
+	DECLARE_DELEGATE_OneParam(FRequirementsChanged, TSharedPtr<FEditableQuery>&);
+
+	class SQueryEditor : public SCompoundWidget
+	{
+	public:
+
+		SLATE_BEGIN_ARGS(SQueryEditor)
+			: _DebuggerModel()
+			, _InitialRequirements()
+			{}
+			SLATE_ARGUMENT(TSharedPtr<FMassDebuggerModel>, DebuggerModel)
+			SLATE_ARGUMENT(FMassFragmentRequirements, InitialRequirements)
+			SLATE_EVENT(FRequirementsChanged, OnRequirementsChanged)
+		SLATE_END_ARGS()
+
+		void Construct(const FArguments& InArgs);
+		void SetQuery(TSharedPtr<FEditableQuery>& InQuery);
+
+	private:
+		void ApplyChanges();
+		void Load();
+		void HandleNameTextCommitted(const FText& InText, ETextCommit::Type InCommitType);
+
+		TSharedRef<ITableRow> OnGenerateFragmentRow(TSharedPtr<FFragmentEntry> InEntry, const TSharedRef<STableViewBase>& OwnerTable);
+
+		TSharedRef<ITableRow> OnGenerateTagRow(TSharedPtr<FFragmentEntry> InEntry, const TSharedRef<STableViewBase>& OwnerTable);
+
+		FReply OnAddFragment();
+		FReply OnRemoveFragment(TSharedPtr<FFragmentEntry> InEntry);
+
+		FReply OnAddTag();
+		FReply OnRemoveTag(TSharedPtr<FFragmentEntry> InEntry);
+
+		void NotifyRequirementsChanged();
+
+		TWeakPtr<FMassDebuggerModel> DebuggerModel;
+
+		FRequirementsChanged OnRequirementsChanged;
+		TSharedPtr<FEditableQuery> Query;
+
+		TArray<TSharedPtr<FFragmentEntry>> FragmentEntries;
+		TArray<TSharedPtr<FFragmentEntry>> TagEntries;
+
+		TArray<TSharedPtr<FString>> AvailableFragmentNames;
+		TMap<FString, TWeakObjectPtr<const UScriptStruct>> FragmentNamesToTypes;
+		TArray<TSharedPtr<FString>> AvailableTagNames;
+		TMap<FString, TWeakObjectPtr<const UScriptStruct>> TagNamesToTypes;
+
+		friend class SFragmentSelectorRow;
+		class SFragmentSelectorRow : public STableRow<TSharedPtr<FFragmentEntry>>
+		{
+		public:
+
+			void Construct(const FArguments& InArgs, const TSharedRef<STableViewBase>& InOwnerTable, const TSharedPtr<FFragmentEntry>& InEntry, TSharedRef<SQueryEditor> InOwnerRequirementsEditor, bool bInShowTags);
+
+			TSharedPtr<SSearchableComboBox> FragmentsComboBox;
+
+			void OnFragmentTypeChanged(TSharedPtr<TWeakObjectPtr<const UScriptStruct>> NewType, ESelectInfo::Type, TSharedPtr<FFragmentEntry> Entry);
+			void OnFragmentAccessChanged(EMassFragmentAccess NewAccess, ESelectInfo::Type, TSharedPtr<FFragmentEntry> Entry);
+			void OnFragmentPresenceChanged(EMassFragmentPresence NewPresence, ESelectInfo::Type, TSharedPtr<FFragmentEntry> Entry);
+			void OnTagTypeChanged(TSharedPtr<TWeakObjectPtr<const UScriptStruct>> NewType, ESelectInfo::Type, TSharedPtr<FFragmentEntry> Entry);
+			void OnTagPresenceChanged(EMassFragmentPresence NewPresence, ESelectInfo::Type, TSharedPtr<FFragmentEntry> Entry);
+
+			FText GetFragmentTypeText(TSharedPtr<FFragmentEntry> Entry) const;
+			FText GetTagTypeText(TSharedPtr<FFragmentEntry> Entry) const;
+			FText GetTagPresenceText(TSharedPtr<FFragmentEntry> Entry) const;
+
+			TWeakPtr<SQueryEditor> OwnerRequirementsEditor;
+			bool bShowTags;
+		};
+
+		TArray<TSharedPtr<FName>> FragmentAccessNames;
+		TArray<TSharedPtr<FName>> FragmentPresenceNames;
+		TMap<FName, EMassFragmentAccess> FragmentAccessMap;
+		TMap<FName, EMassFragmentPresence> FragmentPresenceMap;
+
+		TSharedPtr<SListView<TSharedPtr<FFragmentEntry>>> FragmentListView;
+		TSharedPtr<SListView<TSharedPtr<FFragmentEntry>>> TagListView;
+	};
+
+} // namespace UE::MassDebugger
diff --git a/Engine/Source/Editor/MassEntityDebugger/Private/SMassQueryEditorView.cpp b/Engine/Source/Editor/MassEntityDebugger/Private/SMassQueryEditorView.cpp
new file mode 100644
--- /dev/null
+++ b/Engine/Source/Editor/MassEntityDebugger/Private/SMassQueryEditorView.cpp
@@ -0,0 +1,637 @@
+// Copyright Epic Games, Inc. All Rights Reserved.
+
+#include "SMassQueryEditorView.h"
+#include "MassDebuggerModel.h"
+#include "Dom/JsonObject.h"
+#include "HAL/FileManager.h"
+#include "HAL/PlatformApplicationMisc.h"
+#include "Misc/FileHelper.h"
+#include "Serialization/JsonWriter.h"
+#include "Serialization/JsonReader.h"
+#include "Serialization/JsonSerializer.h"
+#include "Widgets/Layout/SSeparator.h"
+#include "Widgets/SBoxPanel.h"
+#include "Widgets/Input/SButton.h"
+#include "Widgets/Text/STextBlock.h"
+
+#define LOCTEXT_NAMESPACE "SMassQueryEditor"
+
+namespace UE::MassDebugger
+{
+	void SQueryEditorView::Construct(const FArguments& InArgs, TSharedRef<FMassDebuggerModel> InDebuggerModel)
+	{
+		DebuggerModel = InDebuggerModel;
+		InDebuggerModel->OnRefreshDelegate.AddSP(this, &SQueryEditorView::OnRefresh);
+
+		ChildSlot
+		[
+			SNew(SHorizontalBox)
+			+ SHorizontalBox::Slot()
+			.AutoWidth()
+			.Padding(4)
+			[
+				SNew(SVerticalBox)
+				+ SVerticalBox::Slot()
+				.AutoHeight()
+				.Padding(2)
+				[
+					SNew(SHorizontalBox)
+					+ SHorizontalBox::Slot().AutoWidth().Padding(2)
+					[
+						SNew(SButton)
+						.OnClicked(this, &SQueryEditorView::OnSaveQuery)
+						.ToolTipText(LOCTEXT("SaveQuery_Tooltip", "Save selected query to disk"))
+						.Content()
+						[
+							SNew(SImage)
+							.Image(FAppStyle::Get().GetBrush("Icons.Save"))
+						]
+					]
+					+ SHorizontalBox::Slot().AutoWidth().Padding(2)
+					[
+						SNew(SButton)
+						.OnClicked(this, &SQueryEditorView::OnSaveAllQueries)
+						.ToolTipText(LOCTEXT("SaveAllQueries_Tooltip", "Save all queries to disk"))
+						.Content()
+						[
+							SNew(SImage)
+							.Image(FAppStyle::GetBrush("Icons.SaveModified"))
+						]
+					]
+					+ SHorizontalBox::Slot().AutoWidth().Padding(2)
+					[
+						SNew(SButton)
+						.OnClicked(this, &SQueryEditorView::OnCopyQuery)
+						.ToolTipText(LOCTEXT("CopyQuery_Tooltip", "Copy selected query as JSON"))
+						.Content()
+						[
+							SNew(SImage)
+							.Image(FAppStyle::Get().GetBrush("Icons.Clipboard"))
+						]
+					]
+					+ SHorizontalBox::Slot().AutoWidth().Padding(2)
+					[
+						SNew(SButton)
+						.OnClicked(this, &SQueryEditorView::OnPasteQuery)
+						.ToolTipText(LOCTEXT("PasteQuery_Tooltip", "Paste JSON from clipboard as new query"))
+						.Content()
+						[
+							SNew(SImage)
+							.Image(FAppStyle::GetBrush("GenericCommands.Paste"))
+						]
+					]
+				]
+				+ SVerticalBox::Slot()
+				.AutoHeight()
+				.Padding(2)
+				[
+					SNew(SButton)
+					.Text(LOCTEXT("AddQuery", "Add Query"))
+					.OnClicked(this, &SQueryEditorView::OnAddQuery)
+				]
+				+ SVerticalBox::Slot()
+				.AutoHeight()
+				.Padding(2)
+				[
+					SNew(SSeparator)
+				]
+				+ SVerticalBox::Slot()
+				.FillHeight(1.0f)
+				.Padding(2)
+				[
+					SAssignNew(QueryListView, SListView<TSharedPtr<FEditableQuery>>)
+					.ListItemsSource(&QueryList)
+					.SelectionMode(ESelectionMode::Single)
+					.OnGenerateRow(this, &SQueryEditorView::OnGenerateQueryRow)
+					.OnSelectionChanged_Lambda([this](TSharedPtr<FEditableQuery> InQuery, ESelectInfo::Type InSelectType)
+					{
+						if (!InQuery.IsValid())
+						{
+							if (SelectedQuery.IsValid() && QueryList.Contains(SelectedQuery))
+							{
+								QueryListView->SetSelection(SelectedQuery);
+							}
+							else
+							{
+								SelectedQuery = QueryList.Num() > 0 ? QueryList[0] : InQuery;
+								QueryEditor->SetQuery(SelectedQuery);
+							}
+						}
+						else
+						{
+							SelectedQuery = InQuery;
+							QueryEditor->SetQuery(SelectedQuery);
+						}
+					})
+				]
+			]
+			+ SHorizontalBox::Slot()
+			.FillWidth(1.0f)
+			.Padding(4)
+			[
+				SAssignNew(QueryEditor, SQueryEditor)
+				.DebuggerModel(InDebuggerModel)
+				.InitialRequirements(FMassFragmentRequirements())
+				.OnRequirementsChanged(this, &SQueryEditorView::OnRequirementsChanged)
+			]
+		];
+
+		LoadQueriesFromDisk();
+	}
+
+	void SQueryEditorView::OnRefresh()
+	{
+		if (QueryListView.IsValid())
+		{
+			QueryListView->RequestListRefresh();
+		}
+	}
+
+	TSharedRef<ITableRow> SQueryEditorView::OnGenerateQueryRow(TSharedPtr<FEditableQuery> InItem, const TSharedRef<STableViewBase>& OwnerTable)
+	{
+		return SNew(STableRow<TSharedPtr<FEditableQuery>>, OwnerTable)
+		[
+			SNew(SHorizontalBox)
+			+ SHorizontalBox::Slot()
+			.FillWidth(0.3f)
+			.VAlign(VAlign_Center)
+			.Padding(2)
+			[
+				SNew(STextBlock)
+				.Text_Lambda([InItem]()
+				{
+					return InItem.IsValid()
+						? FText::FromString(InItem->Name)
+						: FText::GetEmpty();
+				})
+			]
+			+ SHorizontalBox::Slot()
+			.AutoWidth()
+			.Padding(2)
+			[
+				SNew(SButton)
+				.Text(LOCTEXT("ViewEntities", "View Entities"))
+				.OnClicked(this, &SQueryEditorView::OnViewEntities, InItem)
+			]
+			+ SHorizontalBox::Slot()
+			.AutoWidth()
+			.Padding(2)
+			[
+				SNew(SButton)
+				.OnClicked(this, &SQueryEditorView::OnDeleteQuery, InItem)
+				[
+					SNew(SImage)
+					.Image(FCoreStyle::Get().GetBrush("Icons.Delete"))
+					.ColorAndOpacity(FSlateColor::UseForeground())
+				]
+			]
+		];
+	}
+
+	FReply SQueryEditorView::OnAddQuery()
+	{
+		TSharedPtr<FEditableQuery> NewQuery = MakeShared<FEditableQuery>();
+		NewQuery->Name = TEXT("NewQuery");
+
+		QueryList.Add(NewQuery);
+		if (QueryListView.IsValid())
+		{
+			QueryListView->RequestListRefresh();
+		}
+		return FReply::Handled();
+	}
+
+	FReply SQueryEditorView::OnDeleteQuery(TSharedPtr<FEditableQuery> InItem)
+	{
+		QueryList.Remove(InItem);
+		if (QueryListView.IsValid())
+		{
+			QueryListView->RequestListRefresh();
+		}
+
+		if (SelectedQuery == InItem && QueryEditor.IsValid())
+		{
+			SelectedQuery.Reset();
+			QueryEditor->SetQuery(SelectedQuery);
+		}
+		return FReply::Handled();
+	}
+
+	FReply SQueryEditorView::OnEditQuery(TSharedPtr<FEditableQuery> InItem)
+	{
+		if (!InItem.IsValid() || !QueryEditor.IsValid())
+		{
+			return FReply::Unhandled();
+		}
+
+		SelectedQuery = InItem;
+
+		QueryEditor->SetQuery(SelectedQuery);
+		return FReply::Handled();
+	}
+
+	FReply SQueryEditorView::OnViewEntities(TSharedPtr<FEditableQuery> InItem)
+	{
+		TSharedPtr<FMassDebuggerModel> Model = DebuggerModel.Pin();
+		if (InItem.IsValid() && Model.IsValid() && !Model->IsStale())
+		{
+			FMassEntityQuery Query = InItem->BuildEntityQuery(Model->Environment->GetMutableEntityManager().ToSharedRef());
+			Model->ShowEntitiesView(0, Query);
+		}
+		return FReply::Handled();
+	}
+
+	void SQueryEditorView::OnRequirementsChanged(TSharedPtr<FEditableQuery>& Query)
+	{
+	}
+
+	FEditableQuery FEditableQuery::FromEntityQuery(const FMassEntityQuery& InQuery, FMassDebuggerModel& DebuggerModel)
+	{
+		FEditableQuery OutQuery;
+
+		for (const FMassFragmentRequirementDescription& Desc : InQuery.GetFragmentRequirements())
+		{
+			TSharedPtr<FFragmentEntry> Entry = MakeShared<FFragmentEntry>();
+			Entry->StructType = Desc.StructType;
+			Entry->AccessMode = Desc.AccessMode;
+			Entry->Presence = Desc.Presence;
+			OutQuery.FragmentRequirements.Add(Entry);
+		}
+
+		const FMassTagBitSet& AllTags = InQuery.GetRequiredAllTags();
+		const FMassTagBitSet& AnyTags = InQuery.GetRequiredAnyTags();
+		const FMassTagBitSet& OptTags = InQuery.GetRequiredOptionalTags();
+		const FMassTagBitSet& NoneTags = InQuery.GetRequiredNoneTags();
+
+		for (const TSharedPtr<FMassDebuggerFragmentData>& TagData : DebuggerModel.CachedTagData)
+		{
+			const UScriptStruct& StructType = *TagData->Fragment.Get();
+			EMassFragmentPresence Presence = EMassFragmentPresence::All;
+
+			if (AllTags.Contains(StructType))
+			{
+				Presence = EMassFragmentPresence::All;
+			}
+			else if (AnyTags.Contains(StructType))
+			{
+				Presence = EMassFragmentPresence::Any;
+			}
+			else if (OptTags.Contains(StructType))
+			{
+				Presence = EMassFragmentPresence::Optional;
+			}
+			else if (NoneTags.Contains(StructType))
+			{
+				Presence = EMassFragmentPresence::None;
+			}
+			else
+			{
+				continue;
+			}
+
+			TSharedPtr<FFragmentEntry> Entry = MakeShared<FFragmentEntry>();
+			Entry->StructType = &StructType;
+			Entry->Presence = Presence;
+			OutQuery.TagRequirements.Add(Entry);
+		}
+
+		return OutQuery;
+	}
+
+	FMassEntityQuery FEditableQuery::BuildEntityQuery(const TSharedRef<FMassEntityManager> EntityManager)
+	{
+		FMassEntityQuery Query;
+
+		Query.Initialize(EntityManager);
+
+		for (TSharedPtr<FFragmentEntry>& Fragment : FragmentRequirements)
+		{
+			if (ensure(Fragment->StructType.IsValid()))
+			{
+				Query.AddRequirement(Fragment->StructType.Get(), Fragment->AccessMode, Fragment->Presence);
+			}
+		}
+
+		for (TSharedPtr<FFragmentEntry>& TagEntry : TagRequirements)
+		{
+			if (ensure(TagEntry->StructType.IsValid()))
+			{
+				Query.AddTagRequirement(*TagEntry->StructType, TagEntry->Presence);
+			}
+		}
+		return Query;
+	}
+
+	TSharedPtr<FJsonObject> FFragmentEntry::SerializeToJson() const
+	{
+		if (!StructType.IsValid())
+		{
+			return nullptr;
+		}
+		TSharedPtr<FJsonObject> JsonObject = MakeShared<FJsonObject>();
+
+		JsonObject->SetStringField(TEXT("StructType"), StructType->GetPathName());
+		JsonObject->SetStringField(TEXT("AccessMode"), StaticEnum<EMassFragmentAccess>()->GetNameStringByValue((int64)AccessMode));
+		JsonObject->SetStringField(TEXT("Presence"), StaticEnum<EMassFragmentPresence>()->GetNameStringByValue((int64)Presence));
+
+		return JsonObject;
+	}
+
+	bool FFragmentEntry::DeserializeFromJson(const TSharedPtr<FJsonObject>& JsonObject)
+	{
+		if (!JsonObject.IsValid())
+		{
+			return false;
+		}
+		const FString StructPath = JsonObject->GetStringField(TEXT("StructType"));
+		if (!StructPath.IsEmpty())
+		{
+			UScriptStruct* Found = FindObject<UScriptStruct>(nullptr, *StructPath);
+			StructType = Found;
+		}
+		else
+		{
+			StructType = nullptr;
+		}
+
+		const FString AccessStr = JsonObject->GetStringField(TEXT("AccessMode"));
+		UEnum* AccessEnum = StaticEnum<EMassFragmentAccess>();
+		int64 AccessVal = AccessEnum->GetValueByNameString(AccessStr);
+		AccessMode = (EMassFragmentAccess)AccessVal;
+
+		const FString PresenceStr = JsonObject->GetStringField(TEXT("Presence"));
+		UEnum* PresenceEnum = StaticEnum<EMassFragmentPresence>();
+		int64 PresenceVal = PresenceEnum->GetValueByNameString(PresenceStr);
+		Presence = (EMassFragmentPresence)PresenceVal;
+
+		return true;
+	}
+
+	bool FFragmentEntry::DeserializeFromJsonString(const FString& JsonString)
+	{
+		TSharedPtr<FJsonObject> JsonObject;
+		TSharedRef<TJsonReader<>> Reader = TJsonReaderFactory<>::Create(JsonString);
+
+		if (!FJsonSerializer::Deserialize(Reader, JsonObject) || !JsonObject.IsValid())
+		{
+			return false;
+		}
+
+		return DeserializeFromJson(JsonObject);
+	}
+
+	TSharedPtr<FJsonObject> FEditableQuery::SerializeToJson() const
+	{
+		TSharedPtr<FJsonObject> JsonObject = MakeShared<FJsonObject>();
+		JsonObject->SetStringField(TEXT("Name"), Name);
+
+		TArray<TSharedPtr<FJsonValue>> JsonFragments;
+		JsonFragments.Reserve(FragmentRequirements.Num());
+		for (auto& Entry : FragmentRequirements)
+		{
+			JsonFragments.Add(MakeShared<FJsonValueObject>(Entry->SerializeToJson()));
+		}
+		JsonObject->SetArrayField(TEXT("FragmentRequirements"), MoveTemp(JsonFragments));
+
+		TArray<TSharedPtr<FJsonValue>> JsonTags;
+		JsonTags.Reserve(TagRequirements.Num());
+		for (auto& Entry : TagRequirements)
+		{
+			JsonTags.Add(MakeShared<FJsonValueObject>(Entry->SerializeToJson()));
+		}
+		JsonObject->SetArrayField(TEXT("TagRequirements"), MoveTemp(JsonTags));
+
+		return JsonObject;
+	}
+
+	bool FEditableQuery::DeserializeFromJson(const TSharedPtr<FJsonObject>& JsonObject)
+	{
+		if (!JsonObject.IsValid())
+		{
+			return false;
+		}
+
+		if (JsonObject->HasField(TEXT("Name")))
+		{
+			Name = JsonObject->GetStringField(TEXT("Name"));
+		}
+
+		FragmentRequirements.Reset();
+		const TArray<TSharedPtr<FJsonValue>>* JsonFragments = nullptr;
+		if (JsonObject->TryGetArrayField(TEXT("FragmentRequirements"), JsonFragments))
+		{
+			for (const TSharedPtr<FJsonValue>& Val : *JsonFragments)
+			{
+				if (TSharedPtr<FJsonObject> Obj = Val->AsObject())
+				{
+					TSharedPtr<FFragmentEntry> Entry = MakeShared<FFragmentEntry>();
+					if (!Entry->DeserializeFromJson(Obj))
+					{
+						return false;
+					}
+					Entry->bShowAccessMode = true;
+					FragmentRequirements.Add(MoveTemp(Entry));
+				}
+			}
+		}
+
+		TagRequirements.Reset();
+		const TArray<TSharedPtr<FJsonValue>>* JsonTags = nullptr;
+		if (JsonObject->TryGetArrayField(TEXT("TagRequirements"), JsonTags))
+		{
+			for (const TSharedPtr<FJsonValue>& Val : *JsonTags)
+			{
+				if (TSharedPtr<FJsonObject> Obj = Val->AsObject())
+				{
+					TSharedPtr<FFragmentEntry> Entry = MakeShared<FFragmentEntry>();
+					if (!Entry->DeserializeFromJson(Obj))
+					{
+						return false;
+					}
+					Entry->bShowAccessMode = false;
+					TagRequirements.Add(MoveTemp(Entry));
+				}
+			}
+		}
+
+		return true;
+	}
+
+	bool FEditableQuery::DeserializeFromJsonString(const FString& JsonString)
+	{
+		TSharedPtr<FJsonObject> JsonObject;
+		TSharedRef<TJsonReader<>> Reader = TJsonReaderFactory<>::Create(JsonString);
+
+		if (FJsonSerializer::Deserialize(Reader, JsonObject) && JsonObject.IsValid())
+		{
+			return DeserializeFromJson(JsonObject);
+		}
+
+		return false;
+	}
+
+	FReply SQueryEditorView::OnSaveQuery()
+	{
+		if (!SelectedQuery.IsValid())
+		{
+			return FReply::Handled();
+		}
+
+		const FString FileName = TEXT("Queries.json");
+		const FString SavedDir = FPaths::ProjectSavedDir() / TEXT("MassDebugger");
+		IFileManager::Get().MakeDirectory(*SavedDir, true);
+		const FString FullPath = SavedDir / FileName;
+
+		TSharedPtr<FJsonObject> RootObject = MakeShared<FJsonObject>();
+		TArray<TSharedPtr<FJsonValue>> QueriesArray;
+		if (FPaths::FileExists(FullPath))
+		{
+			FString Existing;
+			FFileHelper::LoadFileToString(Existing, *FullPath);
+			TSharedPtr<FJsonObject> Loaded;
+			const TSharedRef<TJsonReader<>> Reader = TJsonReaderFactory<>::Create(Existing);
+			if (FJsonSerializer::Deserialize(Reader, Loaded) && Loaded.IsValid())
+			{
+				const TArray<TSharedPtr<FJsonValue>>* QueriesArrayJson;;
+				Loaded->TryGetArrayField(TEXT("Queries"), QueriesArrayJson);
+			}
+		}
+
+		const TSharedPtr<FJsonObject> SelJson = SelectedQuery->SerializeToJson();
+		bool bReplaced = false;
+		for (int32 i = 0; i < QueriesArray.Num(); ++i)
+		{
+			if (QueriesArray[i]->AsObject()->GetStringField(TEXT("Name")) == SelectedQuery->Name)
+			{
+				QueriesArray[i] = MakeShared<FJsonValueObject>(SelJson);
+				bReplaced = true;
+				break;
+			}
+		}
+		if (!bReplaced)
+		{
+			QueriesArray.Add(MakeShared<FJsonValueObject>(SelJson));
+		}
+
+		RootObject->SetArrayField(TEXT("Queries"), QueriesArray);
+
+		FString OutString;
+		const TSharedRef<TJsonWriter<>> Writer = TJsonWriterFactory<>::Create(&OutString);
+		FJsonSerializer::Serialize(RootObject.ToSharedRef(), Writer);
+		FFileHelper::SaveStringToFile(OutString, *FullPath);
+
+		return FReply::Handled();
+	}
+
+	FReply SQueryEditorView::OnSaveAllQueries()
+	{
+		const FString FileName = TEXT("Queries.json");
+		const FString SavedDir = FPaths::ProjectSavedDir() / TEXT("MassDebugger");
+		IFileManager::Get().MakeDirectory(*SavedDir, /*Tree=*/true);
+		const FString FullPath = SavedDir / FileName;
+
+		TSharedPtr<FJsonObject> RootObject = MakeShared<FJsonObject>();
+		TArray<TSharedPtr<FJsonValue>> QueriesArray;
+		for (auto& Q : QueryList)
+		{
+			QueriesArray.Add(MakeShared<FJsonValueObject>(Q->SerializeToJson()));
+		}
+		RootObject->SetArrayField(TEXT("Queries"), MoveTemp(QueriesArray));
+
+		FString OutString;
+		const TSharedRef<TJsonWriter<>> Writer = TJsonWriterFactory<>::Create(&OutString);
+		FJsonSerializer::Serialize(RootObject.ToSharedRef(), Writer);
+		FFileHelper::SaveStringToFile(OutString, *FullPath);
+
+		return FReply::Handled();
+	}
+
+	FReply SQueryEditorView::OnCopyQuery()
+	{
+		if (!SelectedQuery.IsValid())
+		{
+			return FReply::Handled();
+		}
+
+		FString JsonString;
+		const TSharedRef<TJsonWriter<>> Writer = TJsonWriterFactory<>::Create(&JsonString);
+		FJsonSerializer::Serialize(SelectedQuery->SerializeToJson().ToSharedRef(), Writer);
+
+		FPlatformApplicationMisc::ClipboardCopy(*JsonString);
+		return FReply::Handled();
+	}
+
+	FReply SQueryEditorView::OnPasteQuery()
+	{
+		FString ClipboardText;
+		FPlatformApplicationMisc::ClipboardPaste(ClipboardText);
+
+		TSharedPtr<FEditableQuery> NewQuery = MakeShared<FEditableQuery>();
+		if (NewQuery->DeserializeFromJsonString(ClipboardText))
+		{
+			QueryList.Add(NewQuery);
+			if (QueryListView.IsValid())
+			{
+				QueryListView->RequestListRefresh();
+			}
+		}
+		return FReply::Handled();
+	}
+
+	void SQueryEditorView::LoadQueriesFromDisk()
+	{
+		const FString SavedDir = FPaths::ProjectSavedDir() / TEXT("MassDebugger");
+		const FString FullPath = SavedDir / TEXT("Queries.json");
+
+		QueryList.Empty();
+
+		if (FPaths::FileExists(FullPath))
+		{
+			FString FileContent;
+			if (FFileHelper::LoadFileToString(FileContent, *FullPath))
+			{
+				TSharedPtr<FJsonObject> RootObj;
+				TSharedRef<TJsonReader<>> Reader = TJsonReaderFactory<>::Create(FileContent);
+				if (FJsonSerializer::Deserialize(Reader, RootObj) && RootObj.IsValid())
+				{
+					const TArray<TSharedPtr<FJsonValue>>* JsonQueries = nullptr;
+					if (RootObj->TryGetArrayField(TEXT("Queries"), JsonQueries))
+					{
+						for (const auto& Val : *JsonQueries)
+						{
+							if (TSharedPtr<FJsonObject> Obj = Val->AsObject())
+							{
+								TSharedPtr<FEditableQuery> NewQ = MakeShared<FEditableQuery>();
+								if (NewQ->DeserializeFromJson(Obj))
+								{
+									QueryList.Add(MoveTemp(NewQ));
+								}
+							}
+						}
+					}
+				}
+			}
+		}
+
+		if (QueryListView.IsValid())
+		{
+			QueryListView->RequestListRefresh();
+		}
+
+		if (QueryList.Num() > 0)
+		{
+			SelectedQuery = QueryList[0];
+			QueryListView->SetSelection(SelectedQuery);
+		}
+		else
+		{
+			SelectedQuery.Reset();
+		}
+		if (QueryEditor.IsValid())
+		{
+			QueryEditor->SetQuery(SelectedQuery);
+		}
+	}
+} // namespace UE::MassDebugger
+
+#undef LOCTEXT_NAMESPACE
diff --git a/Engine/Source/Editor/MassEntityDebugger/Private/SMassQueryEditorView.h b/Engine/Source/Editor/MassEntityDebugger/Private/SMassQueryEditorView.h
new file mode 100644
--- /dev/null
+++ b/Engine/Source/Editor/MassEntityDebugger/Private/SMassQueryEditorView.h
@@ -0,0 +1,75 @@
+// Copyright Epic Games, Inc. All Rights Reserved.
+#pragma once
+
+#include "SMassDebuggerViewBase.h"
+#include "MassEntityQuery.h"
+#include "SMassQueryEditor.h"
+#include "Widgets/Views/SListView.h"
+#include "Widgets/Views/STableRow.h"
+
+class FJsonObject;
+
+namespace UE::MassDebugger
+{
+	struct FFragmentEntry
+	{
+		TSharedPtr<FJsonObject> SerializeToJson() const;
+		bool DeserializeFromJson(const TSharedPtr<FJsonObject>& JsonObject);
+		bool DeserializeFromJsonString(const FString& JsonString);
+		TWeakObjectPtr<const UScriptStruct> StructType = nullptr;
+		EMassFragmentAccess AccessMode = EMassFragmentAccess::ReadOnly;
+		EMassFragmentPresence Presence = EMassFragmentPresence::All;
+		bool bShowAccessMode = true;
+	};
+
+	struct FEditableQuery
+	{
+		TSharedPtr<FJsonObject> SerializeToJson() const;	
+		bool DeserializeFromJson(const TSharedPtr<FJsonObject>& JsonObject);
+		bool DeserializeFromJsonString(const FString& JsonString);
+		static FEditableQuery FromEntityQuery(const FMassEntityQuery& InQuery, FMassDebuggerModel& DebuggerModel);
+		FString Name;
+		TArray<TSharedPtr<FFragmentEntry>> FragmentRequirements;
+		TArray<TSharedPtr<FFragmentEntry>> TagRequirements;
+		FMassEntityQuery BuildEntityQuery(const TSharedRef<FMassEntityManager> EntityManager);
+	};
+
+	class SQueryEditorView : public SMassDebuggerViewBase
+	{
+	public:
+		SLATE_BEGIN_ARGS(SQueryEditorView) {}
+		SLATE_END_ARGS()
+
+		void Construct(const FArguments& InArgs, TSharedRef<FMassDebuggerModel> InDebuggerModel);
+
+		// ~SMassDebuggerViewBase interface
+		virtual void OnRefresh() override;
+		virtual void OnProcessorsSelected(TConstArrayView<TSharedPtr<FMassDebuggerProcessorData>> SelectedProcessors, ESelectInfo::Type SelectInfo) override
+		{}
+		virtual void OnArchetypesSelected(TConstArrayView<TSharedPtr<FMassDebuggerArchetypeData>> SelectedArchetypes, ESelectInfo::Type SelectInfo) override
+		{}
+
+		FReply OnSaveQuery();
+		FReply OnSaveAllQueries();
+		FReply OnCopyQuery();
+		FReply OnPasteQuery();
+		void LoadQueriesFromDisk();
+
+	protected:
+
+		TArray<TSharedPtr<FEditableQuery>> QueryList;
+		TSharedPtr<SListView<TSharedPtr<FEditableQuery>>> QueryListView;
+		TSharedPtr<SQueryEditor> QueryEditor;
+		TSharedPtr<FEditableQuery> SelectedQuery;
+
+		TWeakPtr<FMassDebuggerModel> DebuggerModel;
+
+		TSharedRef<ITableRow> OnGenerateQueryRow(TSharedPtr<FEditableQuery> InItem, const TSharedRef<STableViewBase>& OwnerTable);
+		FReply OnAddQuery();
+		FReply OnDeleteQuery(TSharedPtr<FEditableQuery> InItem);
+		FReply OnEditQuery(TSharedPtr<FEditableQuery> InItem);
+		FReply OnViewEntities(TSharedPtr<FEditableQuery> InItem);
+		void OnRequirementsChanged(TSharedPtr<FEditableQuery>& Query);
+	};
+
+} // namespace UE::MassDebugger
diff --git a/Engine/Source/Runtime/MassEntity/Private/MassDebugger.cpp b/Engine/Source/Runtime/MassEntity/Private/MassDebugger.cpp
--- a/Engine/Source/Runtime/MassEntity/Private/MassDebugger.cpp
+++ b/Engine/Source/Runtime/MassEntity/Private/MassDebugger.cpp
@@ -1389,8 +1389,9 @@ FMassDebugger::FEnvironment* FMassDebugger::GetActiveEnvironment(const FMassEnti
 	return &ActiveEnvironments[Index];
 }
 
-FMassDebugger::FEnvironment::FEnvironment(const FMassEntityManager& InEntityManager)
+FMassDebugger::FEnvironment::FEnvironment(FMassEntityManager& InEntityManager)
 	: EntityManager(InEntityManager.AsWeak())
+	, MutableEntityManager(InEntityManager.AsWeak())
 {
 #if UE_MASS_TRACE_ENABLED
 	TraceStartedDelegateHandle = FTraceAuxiliary::OnTraceStarted.AddLambda([WeakEntityManager = EntityManager](FTraceAuxiliary::EConnectionType TraceType, const FString& TraceDestination)
diff --git a/Engine/Source/Runtime/MassEntity/Public/MassDebugger.h b/Engine/Source/Runtime/MassEntity/Public/MassDebugger.h
--- a/Engine/Source/Runtime/MassEntity/Public/MassDebugger.h
+++ b/Engine/Source/Runtime/MassEntity/Public/MassDebugger.h
@@ -142,6 +142,7 @@ struct FMassDebugger
 	struct FEnvironment
 	{
 		TWeakPtr<const FMassEntityManager> EntityManager;
+		TWeakPtr<FMassEntityManager> MutableEntityManager;
 		TMap<FName, UE::Mass::Debug::FProcessorProviderFunction> ProcessorProviders;
 		FMassEntityHandle SelectedEntity;
 		FMassEntityHandle HighlightedEntity;
@@ -156,7 +157,7 @@ struct FMassDebugger
 		FDelegateHandle TraceStartedDelegateHandle;
 #endif
 
-		explicit FEnvironment(const FMassEntityManager& InEntityManager);
+		explicit FEnvironment(FMassEntityManager& InEntityManager);
 		~FEnvironment();
 
 		bool IsValid() const { return EntityManager.IsValid(); }
