[Mass Debugger]
Properly guard debugger-only code.

diff --git a/Engine/Source/Editor/MassEntityDebugger/Private/MassDebuggerModel.cpp b/Engine/Source/Editor/MassEntityDebugger/Private/MassDebuggerModel.cpp
--- a/Engine/Source/Editor/MassEntityDebugger/Private/MassDebuggerModel.cpp
+++ b/Engine/Source/Editor/MassEntityDebugger/Private/MassDebuggerModel.cpp
@@ -98,6 +98,7 @@ TSharedPtr<const FMassEntityManager> FMassDebuggerEnvironment::GetEntityManager(
 
 TSharedPtr<FMassEntityManager> FMassDebuggerEnvironment::GetMutableEntityManager() const
 {
+#if WITH_MASSENTITY_DEBUG
 	if (TSharedPtr<const FMassEntityManager> SharedConstEntityManager = GetEntityManager())
 	{
 		FMassDebugger::FEnvironment* Environment = FMassDebugger::FindEnvironmentForEntityManager(*SharedConstEntityManager);
@@ -106,6 +107,7 @@ TSharedPtr<FMassEntityManager> FMassDebuggerEnvironment::GetMutableEntityManager
 			return Environment->MutableEntityManager.Pin();
 		}
 	}
+#endif // WITH_MASSENTITY_DEBUG
 	return nullptr;
 }
 
