[Mass Debugger]
Properly guard edotr-only code

diff --git a/Engine/Source/Editor/MassEntityDebugger/Private/MassDebuggerModel.cpp b/Engine/Source/Editor/MassEntityDebugger/Private/MassDebuggerModel.cpp
--- a/Engine/Source/Editor/MassEntityDebugger/Private/MassDebuggerModel.cpp
+++ b/Engine/Source/Editor/MassEntityDebugger/Private/MassDebuggerModel.cpp
@@ -707,16 +707,21 @@ void FMassDebuggerModel::CacheFragmentData(
 
 void FMassDebuggerModel::CacheFragmentData()
 {
+#if WITH_MASSENTITY_DEBUG
 	CacheFragmentData(CachedFragmentData, FMassFragmentBitSet::DebugGetAllStructTypes());
+#endif
 }
 
 void FMassDebuggerModel::CacheTagData()
 {
+#if WITH_MASSENTITY_DEBUG
 	CacheFragmentData(CachedTagData, FMassTagBitSet::DebugGetAllStructTypes());
+#endif
 }
 
 void FMassDebuggerModel::CacheArchetypesData(TMap<FMassArchetypeHandle, TSharedPtr<FMassDebuggerArchetypeData>>& OutTransientArchetypesMap)
 {
+#if WITH_MASSENTITY_DEBUG
 	CachedAllArchetypes.Reset();
 	CachedArchetypeRepresentatives.Reset();
 
@@ -727,6 +732,7 @@ void FMassDebuggerModel::CacheArchetypesData(TMap<FMassArchetypeHandle, TSharedP
 			StoreArchetypes(*EntityManager, OutTransientArchetypesMap);
 		}
 	}
+#endif
 }
 
 float FMassDebuggerModel::MinDistanceToSelectedArchetypes(const TSharedPtr<FMassDebuggerArchetypeData>& InArchetypeData) const
