[Mass Debugger]
Added "Open In Query Editor" button to all query displays in the debugger
Auto-open all New/Opened/Pasted queries in the query editor

diff --git a/Engine/Source/Editor/MassEntityDebugger/Private/MassDebuggerModel.cpp b/Engine/Source/Editor/MassEntityDebugger/Private/MassDebuggerModel.cpp
--- a/Engine/Source/Editor/MassEntityDebugger/Private/MassDebuggerModel.cpp
+++ b/Engine/Source/Editor/MassEntityDebugger/Private/MassDebuggerModel.cpp
@@ -13,6 +13,7 @@
 #include "MassArchetypeData.h"
 #include "Engine/World.h"
 #include "SMassEntitiesView.h"
+#include "SMassQueryEditorView.h"
 #include "SMassDebugger.h"
 
 #define LOCTEXT_NAMESPACE "SMassDebugger"
@@ -119,6 +120,7 @@ FMassDebuggerQueryData::FMassDebuggerQueryData(const FMassEntityQuery& Query, co
 {
 #if WITH_MASSENTITY_DEBUG
 	FMassDebugger::GetQueryExecutionRequirements(Query, ExecutionRequirements);
+	SourceQuery = Query;
 #endif // WITH_MASSENTITY_DEBUG
 }
 
@@ -932,6 +934,24 @@ void FMassDebuggerModel::ResetEntitiesViews()
 	}
 }
 
+void FMassDebuggerModel::RegisterQueryEditor(TSharedRef<UE::MassDebugger::SQueryEditorView> InQueryEditorView)
+{
+	QueryEditorView = InQueryEditorView;
+}
+
+void FMassDebuggerModel::ShowQueryInEditor(const FMassEntityQuery& InQuery, const FString& InQueryName)
+{
+	if (TSharedPtr<SMassDebugger> SharedDebugger = DebuggerWindow.Pin())
+	{
+		SharedDebugger->ShowQueryEditorView();
+	}
+
+	if (TSharedPtr<UE::MassDebugger::SQueryEditorView> SharedEditor = QueryEditorView.Pin())
+	{
+		SharedEditor->ShowQuery(InQuery, InQueryName);
+	}
+}
+
 void FMassDebuggerModel::ResetSelectedArchetypes()
 {
 	for (TSharedPtr<FMassDebuggerArchetypeData>& ArchetypeData : SelectedArchetypes)
diff --git a/Engine/Source/Editor/MassEntityDebugger/Private/MassDebuggerModel.h b/Engine/Source/Editor/MassEntityDebugger/Private/MassDebuggerModel.h
--- a/Engine/Source/Editor/MassEntityDebugger/Private/MassDebuggerModel.h
+++ b/Engine/Source/Editor/MassEntityDebugger/Private/MassDebuggerModel.h
@@ -20,6 +20,11 @@ class UWorld;
 class SMassEntitiesView;
 class SMassDebugger;
 
+namespace UE::MassDebugger
+{
+	class SQueryEditorView;
+}
+
 enum class EMassDebuggerSelectionMode : uint8
 {
 	None,
@@ -55,6 +60,8 @@ struct FMassDebuggerQueryData
 	FText Label;
 	FText AdditionalInformation;
 
+	FMassEntityQuery SourceQuery;
+
 	int32 GetTotalBitsUsedCount();
 	bool IsEmpty() const;
 }; 
@@ -242,6 +249,9 @@ struct FMassDebuggerModel
 	void ResetEntitiesViews();
 	TWeakPtr<SMassEntitiesView> ShowEntitiesView(int32 Index);
 
+	void RegisterQueryEditor(TSharedRef<UE::MassDebugger::SQueryEditorView> InQueryEditorView);
+	void ShowQueryInEditor(const FMassEntityQuery& InQuery, const FString& InQueryName);
+
 protected:
 	void StoreArchetypes(const FMassEntityManager& EntityManager, TMap<FMassArchetypeHandle, TSharedPtr<FMassDebuggerArchetypeData>>& OutTransientArchetypesMap);
 
@@ -294,6 +304,7 @@ public:
 protected:	
 	TArray<TSharedPtr<FMassDebuggerProcessorData>> AllCachedProcessors;
 	TArray<TWeakPtr<SMassEntitiesView>> EntityViews;
+	TWeakPtr<UE::MassDebugger::SQueryEditorView> QueryEditorView;
 	FName SelectedFragmentName;
 
 public:
diff --git a/Engine/Source/Editor/MassEntityDebugger/Private/SMassProcessor.cpp b/Engine/Source/Editor/MassEntityDebugger/Private/SMassProcessor.cpp
--- a/Engine/Source/Editor/MassEntityDebugger/Private/SMassProcessor.cpp
+++ b/Engine/Source/Editor/MassEntityDebugger/Private/SMassProcessor.cpp
@@ -39,7 +39,6 @@ void SMassProcessor::Construct(const FArguments& InArgs, TSharedPtr<FMassDebugge
 			]
 		];
 	}
-#endif //WITH_MASSENTITY_DEBUG
 	
 	if (ProcessorData->ProcessorRequirements->IsEmpty())
 	{
@@ -60,6 +59,7 @@ void SMassProcessor::Construct(const FArguments& InArgs, TSharedPtr<FMassDebugge
 		.AutoHeight()
 		[
 			SNew(SMassQuery, ProcessorData->ProcessorRequirements, InDebuggerModel)
+			.EditName(FString::Format(TEXT("{0}:{1}"), { ProcessorData->Label, ProcessorData->ProcessorRequirements->Label.ToString() }))
 		];
 	}
 
@@ -69,9 +69,12 @@ void SMassProcessor::Construct(const FArguments& InArgs, TSharedPtr<FMassDebugge
 		.AutoHeight()
 		[
 			SNew(SMassQuery, QueryData, InDebuggerModel)
+			.EditName(FString::Format(TEXT("{0}:{1}"), { ProcessorData->Label, QueryData->Label.ToString()}))
 		];
 	}
 
+#endif //WITH_MASSENTITY_DEBUG
+
 	ChildSlot
 	[
 		Box
diff --git a/Engine/Source/Editor/MassEntityDebugger/Private/SMassQuery.cpp b/Engine/Source/Editor/MassEntityDebugger/Private/SMassQuery.cpp
--- a/Engine/Source/Editor/MassEntityDebugger/Private/SMassQuery.cpp
+++ b/Engine/Source/Editor/MassEntityDebugger/Private/SMassQuery.cpp
@@ -47,6 +47,8 @@ void SMassQuery::Construct(const FArguments& InArgs, TSharedPtr<FMassDebuggerQue
 	QueryData = InQueryData;
 
 	FMassExecutionRequirements& ExecutionRequirements = InQueryData->ExecutionRequirements;
+	TWeakPtr<FMassDebuggerModel> WeakDebuggerModel = InDebuggerModel;
+	FString EditName = InArgs._EditName.IsSet() ? InArgs._EditName.Get() : InQueryData->Label.ToString();
 
 	TSharedRef<SVerticalBox> Box = SNew(SVerticalBox);
 	static const FSlateBrush* ReadOnlyBrush = FMassDebuggerStyle::GetBrush("MassDebug.Fragment.ReadOnly");
@@ -54,10 +56,37 @@ void SMassQuery::Construct(const FArguments& InArgs, TSharedPtr<FMassDebuggerQue
 	Box->AddSlot()
 	.AutoHeight()
 	[
-		SNew(SRichTextBlock)
-		.Text(InQueryData->Label)
-		.DecoratorStyleSet(&FAppStyle::Get())
-		.TextStyle(FAppStyle::Get(), "LargeText")
+        SNew(SHorizontalBox)
+
+        + SHorizontalBox::Slot()
+        .AutoWidth()
+        .VAlign(VAlign_Center)
+        [
+            SNew(SRichTextBlock)
+            .Text(InQueryData->Label)
+            .DecoratorStyleSet(&FAppStyle::Get())
+            .TextStyle(FAppStyle::Get(), "LargeText")
+        ]
+        + SHorizontalBox::Slot()
+        .FillWidth(1.0f)
+        [
+            SNullWidget::NullWidget
+        ]
+        + SHorizontalBox::Slot()
+        .AutoWidth()
+        .VAlign(VAlign_Center)
+        [
+            SNew(SButton)
+            .Text(LOCTEXT("OpenInQueryEditor", "Open In Query Editor"))
+            .OnClicked_Lambda([WeakDebuggerModel, QueryDataPtr = InQueryData, EditName]()
+            {
+				if(TSharedPtr<FMassDebuggerModel> SharedModel = WeakDebuggerModel.Pin())
+				{
+					SharedModel->ShowQueryInEditor(QueryDataPtr->SourceQuery, EditName);
+				}
+                return FReply::Handled();
+            })
+        ]
 	];
 
 	AddBitSetPair(Box, ExecutionRequirements.Fragments, TEXT("Fragments"), InDebuggerModel);
diff --git a/Engine/Source/Editor/MassEntityDebugger/Private/SMassQuery.h b/Engine/Source/Editor/MassEntityDebugger/Private/SMassQuery.h
--- a/Engine/Source/Editor/MassEntityDebugger/Private/SMassQuery.h
+++ b/Engine/Source/Editor/MassEntityDebugger/Private/SMassQuery.h
@@ -11,6 +11,7 @@ class SMassQuery : public SCompoundWidget
 public:
 
 	SLATE_BEGIN_ARGS(SMassQuery){}
+		SLATE_ATTRIBUTE(FString, EditName)
 	SLATE_END_ARGS()
 
 	void Construct(const FArguments& InArgs, TSharedPtr<FMassDebuggerQueryData> InQueryData, TSharedRef<FMassDebuggerModel> InDebuggerModel);
diff --git a/Engine/Source/Editor/MassEntityDebugger/Private/SMassQueryEditorView.cpp b/Engine/Source/Editor/MassEntityDebugger/Private/SMassQueryEditorView.cpp
--- a/Engine/Source/Editor/MassEntityDebugger/Private/SMassQueryEditorView.cpp
+++ b/Engine/Source/Editor/MassEntityDebugger/Private/SMassQueryEditorView.cpp
@@ -137,6 +137,7 @@ namespace UE::MassDebugger
 		];
 
 		LoadQueriesFromDisk();
+		InDebuggerModel->RegisterQueryEditor(SharedThis(this));
 	}
 
 	void SQueryEditorView::OnRefresh()
@@ -197,6 +198,8 @@ namespace UE::MassDebugger
 		if (QueryListView.IsValid())
 		{
 			QueryListView->RequestListRefresh();
+			QueryListView->SetSelection(NewQuery);
+			QueryListView->RequestScrollIntoView(NewQuery);
 		}
 		return FReply::Handled();
 	}
@@ -245,9 +248,9 @@ namespace UE::MassDebugger
 	{
 	}
 
-	FEditableQuery FEditableQuery::FromEntityQuery(const FMassEntityQuery& InQuery, FMassDebuggerModel& DebuggerModel)
+	void FEditableQuery::InitFromEntityQuery(const FMassEntityQuery& InQuery, FMassDebuggerModel& DebuggerModel)
 	{
-		FEditableQuery OutQuery;
+		FragmentRequirements.Reset();
 
 		for (const FMassFragmentRequirementDescription& Desc : InQuery.GetFragmentRequirements())
 		{
@@ -255,12 +258,12 @@ namespace UE::MassDebugger
 			Entry->StructType = Desc.StructType;
 			Entry->AccessMode = Desc.AccessMode;
 			Entry->Presence = Desc.Presence;
-			OutQuery.FragmentRequirements.Add(Entry);
+			FragmentRequirements.Add(Entry);
 		}
 
 		const FMassTagBitSet& AllTags = InQuery.GetRequiredAllTags();
 		const FMassTagBitSet& AnyTags = InQuery.GetRequiredAnyTags();
-		const FMassTagBitSet& OptTags = InQuery.GetRequiredOptionalTags();
+		const FMassTagBitSet& OptionalTags = InQuery.GetRequiredOptionalTags();
 		const FMassTagBitSet& NoneTags = InQuery.GetRequiredNoneTags();
 
 		for (const TSharedPtr<FMassDebuggerFragmentData>& TagData : DebuggerModel.CachedTagData)
@@ -276,7 +279,7 @@ namespace UE::MassDebugger
 			{
 				Presence = EMassFragmentPresence::Any;
 			}
-			else if (OptTags.Contains(StructType))
+			else if (OptionalTags.Contains(StructType))
 			{
 				Presence = EMassFragmentPresence::Optional;
 			}
@@ -292,10 +295,8 @@ namespace UE::MassDebugger
 			TSharedPtr<FFragmentEntry> Entry = MakeShared<FFragmentEntry>();
 			Entry->StructType = &StructType;
 			Entry->Presence = Presence;
-			OutQuery.TagRequirements.Add(Entry);
+			TagRequirements.Add(Entry);
 		}
-
-		return OutQuery;
 	}
 
 	FMassEntityQuery FEditableQuery::BuildEntityQuery(const TSharedRef<FMassEntityManager> EntityManager)
@@ -470,6 +471,23 @@ namespace UE::MassDebugger
 		return false;
 	}
 
+	void SQueryEditorView::ShowQuery(const FMassEntityQuery& InQuery, const FString& InQueryName)
+	{
+		if (TSharedPtr<FMassDebuggerModel> SharedModel = DebuggerModel.Pin())
+		{
+			TSharedPtr<FEditableQuery> NewQuery = MakeShared<FEditableQuery>();
+			NewQuery->Name = InQueryName;
+			NewQuery->InitFromEntityQuery(InQuery, *SharedModel);
+			QueryList.Add(NewQuery);
+			if (QueryListView.IsValid())
+			{
+				QueryListView->RequestListRefresh();
+				QueryListView->SetSelection(NewQuery);
+				QueryListView->RequestScrollIntoView(NewQuery);
+			}
+		}
+	}
+
 	FReply SQueryEditorView::OnSaveQuery()
 	{
 		if (!SelectedQuery.IsValid())
@@ -573,6 +591,8 @@ namespace UE::MassDebugger
 			if (QueryListView.IsValid())
 			{
 				QueryListView->RequestListRefresh();
+				QueryListView->SetSelection(NewQuery);
+				QueryListView->RequestScrollIntoView(NewQuery);
 			}
 		}
 		return FReply::Handled();
diff --git a/Engine/Source/Editor/MassEntityDebugger/Private/SMassQueryEditorView.h b/Engine/Source/Editor/MassEntityDebugger/Private/SMassQueryEditorView.h
--- a/Engine/Source/Editor/MassEntityDebugger/Private/SMassQueryEditorView.h
+++ b/Engine/Source/Editor/MassEntityDebugger/Private/SMassQueryEditorView.h
@@ -27,7 +27,7 @@ namespace UE::MassDebugger
 		TSharedPtr<FJsonObject> SerializeToJson() const;	
 		bool DeserializeFromJson(const TSharedPtr<FJsonObject>& JsonObject);
 		bool DeserializeFromJsonString(const FString& JsonString);
-		static FEditableQuery FromEntityQuery(const FMassEntityQuery& InQuery, FMassDebuggerModel& DebuggerModel);
+		void InitFromEntityQuery(const FMassEntityQuery& InQuery, FMassDebuggerModel& DebuggerModel);
 		FString Name;
 		TArray<TSharedPtr<FFragmentEntry>> FragmentRequirements;
 		TArray<TSharedPtr<FFragmentEntry>> TagRequirements;
@@ -49,6 +49,8 @@ namespace UE::MassDebugger
 		virtual void OnArchetypesSelected(TConstArrayView<TSharedPtr<FMassDebuggerArchetypeData>> SelectedArchetypes, ESelectInfo::Type SelectInfo) override
 		{}
 
+		void ShowQuery(const FMassEntityQuery& InQuery, const FString& InQueryName);
+
 		FReply OnSaveQuery();
 		FReply OnSaveAllQueries();
 		FReply OnCopyQuery();
