[MassDebugger]
Fixed localization key collision

diff --git a/Engine/Source/Editor/MassEntityDebugger/Private/SMassQueryEditor.cpp b/Engine/Source/Editor/MassEntityDebugger/Private/SMassQueryEditor.cpp
--- a/Engine/Source/Editor/MassEntityDebugger/Private/SMassQueryEditor.cpp
+++ b/Engine/Source/Editor/MassEntityDebugger/Private/SMassQueryEditor.cpp
@@ -294,7 +294,7 @@ namespace UE::MassDebugger
 						{
 							return InEntry->StructType->GetDisplayNameText();
 						}
-						return LOCTEXT("None", "<None>");
+						return LOCTEXT("NoneFragment", "<None>");
 					})
 				]
 
