[MassDebugger]
Fixed loctext key collisions

diff --git a/Engine/Source/Editor/MassEntityDebugger/Private/SMassQueryEditor.cpp b/Engine/Source/Editor/MassEntityDebugger/Private/SMassQueryEditor.cpp
--- a/Engine/Source/Editor/MassEntityDebugger/Private/SMassQueryEditor.cpp
+++ b/Engine/Source/Editor/MassEntityDebugger/Private/SMassQueryEditor.cpp
@@ -250,7 +250,7 @@ namespace UE::MassDebugger
 					case EMassFragmentAccess::ReadWrite:
 						return LOCTEXT("ReadWrite", "ReadWrite");
 					default:
-						return LOCTEXT("None", "<None>");
+						return LOCTEXT("NoneAccess", "<None>");
 					}
 				})
 			];
@@ -325,7 +325,7 @@ namespace UE::MassDebugger
 				[
 					SNew(STextBlock).Text_Lambda([InEntry]()
 					{
-						if (!InEntry.IsValid()) return LOCTEXT("None", "<None>");
+						if (!InEntry.IsValid()) return LOCTEXT("NonePres", "<None>");
 						switch (InEntry->Presence)
 						{
 						case EMassFragmentPresence::All:
@@ -454,14 +454,14 @@ namespace UE::MassDebugger
 	{
 		return Entry.IsValid() && Entry->StructType.IsValid()
 			? FText::FromName(Entry->StructType->GetFName())
-			: LOCTEXT("None", "<None>");
+			: LOCTEXT("NoneFragment", "<None>");
 	}
 
 	FText SQueryEditor::SFragmentSelectorRow::GetTagTypeText(TSharedPtr<FFragmentEntry> Entry) const
 	{
 		return Entry.IsValid() && Entry->StructType.IsValid()
 			? FText::FromName(Entry->StructType->GetFName())
-			: LOCTEXT("None", "<None>");
+			: LOCTEXT("NoneTag", "<None>");
 	}
 
 	FText SQueryEditor::SFragmentSelectorRow::GetTagPresenceText(TSharedPtr<FFragmentEntry> Entry) const
