[MassDebugger]
Fix crash bug when stale fragment data pointers are refreshed in the entity view.

diff --git a/Engine/Source/Editor/MassEntityDebugger/Private/SMassEntitiesList.cpp b/Engine/Source/Editor/MassEntityDebugger/Private/SMassEntitiesList.cpp
--- a/Engine/Source/Editor/MassEntityDebugger/Private/SMassEntitiesList.cpp
+++ b/Engine/Source/Editor/MassEntityDebugger/Private/SMassEntitiesList.cpp
@@ -253,15 +253,37 @@ void SMassEntitiesList::RefreshFragmentData()
 				{
 					if (UE::Mass::IsA<FMassFragment>(Info.StructType))
 					{
-						FMassDebugger::GetFragmentData(EntityManager, Info.StructType, Row.Entity, Info.StructData);
+						if (Info.StructData)
+						{
+							FMassDebugger::GetFragmentData(EntityManager, Info.StructType, Row.Entity, Info.StructData);
+						}
+						else
+						{
+							Info.StructData = FMassDebugger::GetFragmentData(EntityManager, Info.StructType, Row.Entity);
+						}
 					}
 					else if (UE::Mass::IsA<FMassSharedFragment>(Info.StructType))
 					{
-						FMassDebugger::GetSharedFragmentData(EntityManager, Info.StructType, Row.Entity, Info.StructData);
+						if (Info.StructData)
+						{
+							FMassDebugger::GetSharedFragmentData(EntityManager, Info.StructType, Row.Entity, Info.StructData);
+						}
+						else
+						{
+							Info.StructData = FMassDebugger::GetSharedFragmentData(EntityManager, Info.StructType, Row.Entity);
+						}
+						
 					}
 					else if (UE::Mass::IsA<FMassConstSharedFragment>(Info.StructType))
 					{
-						FMassDebugger::GetConstSharedFragmentData(EntityManager, Info.StructType, Row.Entity, Info.StructData);
+						if (Info.StructData)
+						{
+							FMassDebugger::GetConstSharedFragmentData(EntityManager, Info.StructType, Row.Entity, Info.StructData);
+						}
+						else
+						{
+							Info.StructData = FMassDebugger::GetConstSharedFragmentData(EntityManager, Info.StructType, Row.Entity);
+						}
 					}
 					else
 					{
