[MassDebugger]

Refactor breakpoints to be more extensible.

Added filtering options to all breakpoint types including by editable query, gameplay debugger selection, etc.

Breakpoints trigger for commands on push to catch the source of the command which is lost by the time the command executes.

Added hit count to breakpoints that increment on trigger even if the actual break is disabled.

Added cvar to auto-selse new debugger environments so breakpoints can trigger for init events and first-frame execution.

New breakpoint types:
* Entity create/destroy
* Tag add/remove
* Fragment add/remove

New Breakpoint editor interface
* Breakpoints persist between editor/PIE runs
* Supports copy/paste

diff --git a/Engine/Source/Editor/MassEntityDebugger/Private/MassDebuggerModel.cpp b/Engine/Source/Editor/MassEntityDebugger/Private/MassDebuggerModel.cpp
--- a/Engine/Source/Editor/MassEntityDebugger/Private/MassDebuggerModel.cpp
+++ b/Engine/Source/Editor/MassEntityDebugger/Private/MassDebuggerModel.cpp
@@ -15,6 +15,12 @@
 #include "SMassEntitiesView.h"
 #include "SMassQueryEditorView.h"
 #include "SMassDebugger.h"
+#include "Algo/AnyOf.h"
+#include "Misc/FileHelper.h"
+#include "Serialization/JsonSerializer.h"
+#include "Serialization/JsonReader.h"
+#include "Serialization/JsonWriter.h"
+#include "Dom/JsonObject.h"
 
 #define LOCTEXT_NAMESPACE "SMassDebugger"
 
@@ -427,6 +433,9 @@ FMassDebuggerModel::FMassDebuggerModel()
 {
 #if WITH_MASSENTITY_DEBUG
 	OnEntitySelectedHandle = FMassDebugger::OnEntitySelectedDelegate.AddRaw(this, &FMassDebuggerModel::OnEntitySelected);
+	OnBreakpointsChangedHandle = FMassDebugger::OnBreakpointsChangedDelegate.AddRaw(this, &FMassDebuggerModel::OnBreakpointsChanged);
+	LoadQueriesFromDisk();
+	LoadBreakpointsFromDisk();
 #endif // WITH_MASSENTITY_DEBUG
 }
 
@@ -474,9 +483,11 @@ void FMassDebuggerModel::RefreshAll()
 		CacheProcessingGraphs(CompositeProcessors);
 		CacheFragmentData();
 		CacheTagData();
-
 		ClearArchetypeSelection();
 
+		ReconcileAllBreakpoints();
+		ApplyBreakpointsToCurrentEnvironment();
+
 		OnRefreshDelegate.Broadcast();
 	}
 #endif // WITH_MASSENTITY_DEBUG
@@ -563,6 +574,8 @@ void FMassDebuggerModel::CacheProcessorsData(const TMap<FMassArchetypeHandle, TS
 
 	CachedProcessorCollections.Reset();
 	AllCachedProcessors.Reset();
+	ProcessorNames.Reset();
+	ProcessorMap.Reset();
 	
 	if (Environment == nullptr)
 	{
@@ -631,6 +644,12 @@ void FMassDebuggerModel::CacheProcessorsData(const TMap<FMassArchetypeHandle, TS
 	}
 
 	AllCachedProcessors.Sort(SortPredicate);
+
+	for (TSharedPtr<FMassDebuggerProcessorData>& Processor : AllCachedProcessors)
+	{
+		ProcessorNames.Add(MakeShared<FString>(Processor->Name));
+		ProcessorMap.Add(Processor->Name, Processor);
+	}
 }
 
 void FMassDebuggerModel::CacheProcessingGraphs(TConstArrayView<TNotNull<const UMassCompositeProcessor*>> InCompositeProcessors)
@@ -711,13 +730,43 @@ void FMassDebuggerModel::CacheFragmentData()
 {
 #if WITH_MASSENTITY_DEBUG
 	CacheFragmentData(CachedFragmentData, FMassFragmentBitSet::DebugGetAllStructTypes());
+	FragmentNames.Reset();
+	FragmentMap.Reset();
+
+	for (TSharedPtr<FMassDebuggerFragmentData>& Fragment : CachedFragmentData)
+	{
+		FragmentNames.Add(MakeShared<FString>(Fragment->Name.ToString()));
+		FragmentMap.Add(Fragment->Name.ToString(), Fragment);
+	}
+#endif
+}
+
+TSharedPtr<FMassDebuggerFragmentData> FMassDebuggerModel::FindFragmentData(const UScriptStruct* FragmentType)
+{
+#if WITH_MASSENTITY_DEBUG
+	for (TSharedPtr<FMassDebuggerFragmentData>& FragmentDataPtr : CachedFragmentData)
+	{
+		if (FragmentDataPtr->Fragment.Get() == FragmentType)
+		{
+			return FragmentDataPtr;
+		}
+	}
 #endif
+	return nullptr;
 }
 
 void FMassDebuggerModel::CacheTagData()
 {
 #if WITH_MASSENTITY_DEBUG
 	CacheFragmentData(CachedTagData, FMassTagBitSet::DebugGetAllStructTypes());
+	TagNames.Reset();
+	TagMap.Reset();
+
+	for (TSharedPtr<FMassDebuggerFragmentData>& Tag : CachedTagData)
+	{
+		TagNames.Add(MakeShared<FString>(Tag->Name.ToString()));
+		TagMap.Add(Tag->Name.ToString(), Tag);
+	}
 #endif
 }
 
@@ -737,6 +786,429 @@ void FMassDebuggerModel::CacheArchetypesData(TMap<FMassArchetypeHandle, TSharedP
 #endif
 }
 
+void FMassDebuggerModel::LoadBreakpointsFromDisk()
+{
+#if WITH_MASSENTITY_DEBUG
+	CachedBreakpoints.Reset();
+	FMassDebugger::ClearAllBreakpoints();
+	
+	const FString SavedDir = FPaths::ProjectSavedDir() / TEXT("MassDebugger");
+	const FString FullPath = SavedDir / TEXT("Breakpoints.json");
+
+	if (FPaths::FileExists(FullPath))
+	{
+		FString FileContent;
+		if (FFileHelper::LoadFileToString(FileContent, *FullPath))
+		{
+			TSharedPtr<FJsonObject> RootObj;
+			TSharedRef<TJsonReader<>> Reader = TJsonReaderFactory<>::Create(FileContent);
+			if (FJsonSerializer::Deserialize(Reader, RootObj) && RootObj.IsValid())
+			{
+				const TArray<TSharedPtr<FJsonValue>>* JsonBreakpoints = nullptr;
+				if (RootObj->TryGetArrayField(TEXT("Breakpoints"), JsonBreakpoints))
+				{
+					for (const TSharedPtr<FJsonValue>& Val : *JsonBreakpoints)
+					{
+						if (TSharedPtr<FJsonObject> Obj = Val->AsObject())
+						{
+							TSharedPtr<FMassDebuggerBreakpointData> NewBp = MakeShared<FMassDebuggerBreakpointData>();
+							if (NewBp->DeserializeFromJson(Obj))
+							{
+								CachedBreakpoints.Add(NewBp);
+							}
+						}
+					}
+				}
+			}
+		}
+	}
+	OnEditorBreakpointsChangedDelegate.Broadcast();
+#endif
+}
+
+void FMassDebuggerModel::SaveBreakpointsToDisk()
+{
+	const FString FileName = TEXT("Breakpoints.json");
+	const FString SavedDir = FPaths::ProjectSavedDir() / TEXT("MassDebugger");
+	IFileManager::Get().MakeDirectory(*SavedDir, /*Tree=*/true);
+	const FString FullPath = SavedDir / FileName;
+
+	TSharedPtr<FJsonObject> RootObject = MakeShared<FJsonObject>();
+	TArray<TSharedPtr<FJsonValue>> BreakpointsArray;
+	for (TSharedPtr<FMassDebuggerBreakpointData>& Breakpoint : CachedBreakpoints)
+	{
+		BreakpointsArray.Add(MakeShared<FJsonValueObject>(Breakpoint->SerializeToJson()));
+	}
+	RootObject->SetArrayField(TEXT("Breakpoints"), MoveTemp(BreakpointsArray));
+
+	FString OutString;
+	const TSharedRef<TJsonWriter<>> Writer = TJsonWriterFactory<>::Create(&OutString);
+	FJsonSerializer::Serialize(RootObject.ToSharedRef(), Writer);
+	FFileHelper::SaveStringToFile(OutString, *FullPath);
+}
+
+void FMassDebuggerModel::CreateBreakpoint()
+{
+	TSharedPtr<FMassDebuggerBreakpointData> NewBp = MakeShared<FMassDebuggerBreakpointData>();
+	CachedBreakpoints.Add(NewBp);
+	if (!IsStale())
+	{
+		NewBp->ApplyToEngine(*this);
+	}
+	OnEditorBreakpointsChangedDelegate.Broadcast();
+}
+
+void FMassDebuggerModel::CreateBreakpointFromString(const FString& InString)
+{
+	TSharedPtr<FMassDebuggerBreakpointData> NewBp = MakeShared<FMassDebuggerBreakpointData>();
+	if (NewBp->DeserializeFromJsonString(InString))
+	{
+		CachedBreakpoints.Add(NewBp);
+		if (!IsStale())
+		{
+			NewBp->ApplyToEngine(*this);
+		}
+		OnEditorBreakpointsChangedDelegate.Broadcast();
+	}
+}
+
+void FMassDebuggerModel::ApplyBreakpointsToCurrentEnvironment()
+{
+	if (!IsStale())
+	{
+		FMassDebugger::ClearAllBreakpoints();
+		const FMassEntityManager& EntityManager = *Environment->EntityManager.Pin();
+		for (TSharedPtr<FMassDebuggerBreakpointData>& Breakpoint : CachedBreakpoints)
+		{
+			// apply breakpoint but wait until all are added before refreshing
+			Breakpoint->ApplyToEngine(*this, false);
+		}
+		FMassDebugger::RefreshBreakpoints();
+	}
+	OnEditorBreakpointsChangedDelegate.Broadcast();
+}
+
+void FMassDebuggerModel::ReconcileAllBreakpoints()
+{
+	for (TSharedPtr<FMassDebuggerBreakpointData>& Breakpoint : CachedBreakpoints)
+	{
+		Breakpoint->ReconcileDataFromNames(*this);
+	}
+	OnEditorBreakpointsChangedDelegate.Broadcast();
+}
+
+void FMassDebuggerModel::RemoveAllBreakpoints()
+{
+	CachedBreakpoints.Reset();
+	ApplyBreakpointsToCurrentEnvironment();
+}
+
+void FMassDebuggerModel::RemoveBreakpoint(UE::Mass::Debug::FBreakpointHandle Handle)
+{
+	for (int32 i = 0; i < CachedBreakpoints.Num(); i++)
+	{
+		if (CachedBreakpoints[i]->BreakpointInstance.Handle == Handle)
+		{
+			CachedBreakpoints.RemoveAt(i);
+			ApplyBreakpointsToCurrentEnvironment();
+			return;
+		}
+	}
+}
+
+FMassDebuggerBreakpointData::FMassDebuggerBreakpointData(const UE::Mass::Debug::FBreakpoint& InBreakpoint, FMassDebuggerModel& Model)
+{
+	BreakpointInstance = InBreakpoint;
+	ReconcileDataFromEngineBreakpoint(Model);
+}
+
+void FMassDebuggerBreakpointData::ReconcileDataFromNames(FMassDebuggerModel& Model)
+{
+	switch (BreakpointInstance.TriggerType)
+	{
+	case UE::Mass::Debug::FBreakpoint::ETriggerType::ProcessorExecute:
+	{
+		if (TWeakPtr<FMassDebuggerProcessorData>* Found = Model.ProcessorMap.Find(TriggerName))
+		{
+			if ((TriggerProcessor = (*Found).Pin()))
+			{
+				if (const UMassProcessor* Processor = TriggerProcessor->Processor.Get())
+				{
+					BreakpointInstance.Trigger.Set<TObjectKey<const UMassProcessor>>(Processor);
+				}
+			}
+		}
+		break;
+	}
+	case UE::Mass::Debug::FBreakpoint::ETriggerType::FragmentWrite:
+	case UE::Mass::Debug::FBreakpoint::ETriggerType::FragmentAdd:
+	case UE::Mass::Debug::FBreakpoint::ETriggerType::FragmentRemove:
+	case UE::Mass::Debug::FBreakpoint::ETriggerType::TagAdd:
+	case UE::Mass::Debug::FBreakpoint::ETriggerType::TagRemove:
+	{
+		if (TWeakPtr<FMassDebuggerFragmentData>* Found = Model.FragmentMap.Find(TriggerName))
+		{
+			if ((TriggerFragment = (*Found).Pin()))
+			{
+				if (const UScriptStruct* Fragment = TriggerFragment->Fragment.Get())
+				{
+					BreakpointInstance.Trigger.Set<TObjectKey<const UScriptStruct>>(Fragment);
+				}
+			}
+		}
+		break;
+	}
+	default:
+		break;
+	}
+	
+	if(BreakpointInstance.FilterType == UE::Mass::Debug::FBreakpoint::EFilterType::Query)
+	{
+		if (TWeakPtr<UE::MassDebugger::FEditableQuery>* Found = Model.QueryMap.Find(FilterName))
+		{
+			FilterQuery = (*Found).Pin();
+			if (!Model.IsStale())
+			{
+				TSharedPtr<FMassEntityManager> EntityManager = Model.Environment->GetMutableEntityManager();
+				BreakpointInstance.Filter.Set<FMassFragmentRequirements>(FilterQuery->BuildEntityQuery(EntityManager.ToSharedRef()));
+			}
+		}
+	}
+}
+
+void FMassDebuggerBreakpointData::ReconcileDataFromEngineBreakpoint(FMassDebuggerModel& Model)
+{
+	switch (BreakpointInstance.TriggerType)
+	{
+	case UE::Mass::Debug::FBreakpoint::ETriggerType::ProcessorExecute:
+		if (BreakpointInstance.Trigger.IsType<TObjectKey<const UMassProcessor>>())
+		{
+			const UMassProcessor* Proc = BreakpointInstance.Trigger.Get<TObjectKey<const UMassProcessor>>().ResolveObjectPtr();
+
+			if (Proc)
+			{
+				TriggerProcessor = Model.GetProcessorDataChecked(*Proc);
+				TriggerName = TriggerProcessor->Name;
+			}
+		}
+		break;
+
+	case UE::Mass::Debug::FBreakpoint::ETriggerType::FragmentWrite:
+	case UE::Mass::Debug::FBreakpoint::ETriggerType::FragmentAdd:
+	case UE::Mass::Debug::FBreakpoint::ETriggerType::FragmentRemove:
+	case UE::Mass::Debug::FBreakpoint::ETriggerType::TagAdd:
+	case UE::Mass::Debug::FBreakpoint::ETriggerType::TagRemove:
+
+		if (BreakpointInstance.Trigger.IsType<TObjectKey<const UScriptStruct>>())
+		{
+			const UScriptStruct* StructType = BreakpointInstance.Trigger.Get<TObjectKey<const UScriptStruct>>().ResolveObjectPtr();
+			TriggerFragment = Model.FindFragmentData(StructType);
+			TriggerName = TriggerFragment->Name.ToString();
+		}
+		break;
+
+	default:
+		break;
+	}
+}
+
+void FMassDebuggerBreakpointData::ApplyToEngine(FMassDebuggerModel& Model, bool bRefreshEngineBreakpoints/* = true*/)
+{
+#if WITH_MASSENTITY_DEBUG
+	if (Model.IsStale())
+	{
+		return;
+	}
+	
+	TSharedPtr<FMassEntityManager> EntityManager = Model.Environment->GetMutableEntityManager();
+
+	if (FMassDebugger::FEnvironment* Environment = FMassDebugger::FindEnvironmentForEntityManager(*EntityManager))
+	{
+		if (BreakpointInstance.FilterType == UE::Mass::Debug::FBreakpoint::EFilterType::Query
+			&& !BreakpointInstance.Filter.IsType<FMassFragmentRequirements>())
+		{
+			BreakpointInstance.Filter.Set<FMassFragmentRequirements>(FilterQuery->BuildEntityQuery(EntityManager.ToSharedRef()));
+		}
+
+		if (UE::Mass::Debug::FBreakpoint* Breakpoint = Environment->FindBreakpointByHandle(BreakpointInstance.Handle))
+		{
+			*Breakpoint = BreakpointInstance;
+		}
+		else
+		{
+			UE::Mass::Debug::FBreakpoint& NewBreakpoint = FMassDebugger::CreateBreakpoint(*EntityManager);
+			NewBreakpoint = BreakpointInstance;
+		}
+		if (bRefreshEngineBreakpoints)
+		{
+			FMassDebugger::RefreshBreakpoints();
+		}
+	}
+#endif
+}
+
+TSharedPtr<FJsonObject> FMassDebuggerBreakpointData::SerializeToJson() const
+{
+	TSharedPtr<FJsonObject> JsonObject = MakeShared<FJsonObject>();
+
+	JsonObject->SetBoolField(TEXT("Enabled"), BreakpointInstance.bEnabled);
+	JsonObject->SetStringField(TEXT("TriggerType"), UE::Mass::Debug::FBreakpoint::TriggerTypeToString(BreakpointInstance.TriggerType));
+	JsonObject->SetStringField(TEXT("FilterType"),UE::Mass::Debug::FBreakpoint::FilterTypeToString(BreakpointInstance.FilterType));
+
+	switch (BreakpointInstance.TriggerType)
+	{
+	case UE::Mass::Debug::FBreakpoint::ETriggerType::ProcessorExecute:
+		if (TriggerProcessor.IsValid())
+		{
+			JsonObject->SetStringField(TEXT("ProcessorName"),TriggerProcessor->Name);
+		}
+		else
+		{
+			JsonObject->SetStringField(TEXT("ProcessorName"), TriggerName);
+		}
+		break;
+
+	case UE::Mass::Debug::FBreakpoint::ETriggerType::FragmentWrite:
+	case UE::Mass::Debug::FBreakpoint::ETriggerType::FragmentAdd:
+	case UE::Mass::Debug::FBreakpoint::ETriggerType::FragmentRemove:
+	case UE::Mass::Debug::FBreakpoint::ETriggerType::TagAdd:
+	case UE::Mass::Debug::FBreakpoint::ETriggerType::TagRemove:
+		if (TriggerFragment.IsValid())
+		{
+			JsonObject->SetStringField(TEXT("FragmentName"),TriggerFragment->Name.ToString());
+		}
+		else
+		{
+			JsonObject->SetStringField(TEXT("FragmentName"), TriggerName);
+		}
+		break;
+
+	default:
+		break;
+	}
+
+	switch (BreakpointInstance.FilterType)
+	{
+	case UE::Mass::Debug::FBreakpoint::EFilterType::SpecificEntity:
+	{
+		const FMassEntityHandle& Handle = BreakpointInstance.Filter.Get<FMassEntityHandle>();
+		JsonObject->SetNumberField(TEXT("EntityIndex"), Handle.Index);
+		JsonObject->SetNumberField(TEXT("EntitySerialNumber"), Handle.SerialNumber);
+		break;
+	}
+	case UE::Mass::Debug::FBreakpoint::EFilterType::Query:
+		if (FilterQuery.IsValid())
+		{
+			JsonObject->SetStringField(TEXT("FilterQueryName"),FilterQuery->Name);
+		}
+		else
+		{
+			JsonObject->SetStringField(TEXT("FilterQueryName"), FilterName);
+		}
+		break;
+
+	default:
+		break;
+	}
+
+	return JsonObject;
+}
+
+bool FMassDebuggerBreakpointData::DeserializeFromJson(const TSharedPtr<FJsonObject>& JsonObject)
+{
+	if (!JsonObject.IsValid())
+	{
+		return false;
+	}
+
+	BreakpointInstance.bEnabled = false;
+	JsonObject->TryGetBoolField(TEXT("Enabled"), BreakpointInstance.bEnabled);
+
+	{
+		FString TriggerTypeStr;
+		if (JsonObject->TryGetStringField(TEXT("TriggerType"), TriggerTypeStr))
+		{
+			UE::Mass::Debug::FBreakpoint::StringToTriggerType(TriggerTypeStr, BreakpointInstance.TriggerType);
+		}
+	}
+
+	TriggerProcessor.Reset();
+	TriggerFragment.Reset();
+
+	switch (BreakpointInstance.TriggerType)
+	{
+	case UE::Mass::Debug::FBreakpoint::ETriggerType::ProcessorExecute:
+	{
+		if (!JsonObject->TryGetStringField(TEXT("ProcessorName"), TriggerName))
+		{
+			return false;
+		}
+		break;
+	}
+	case UE::Mass::Debug::FBreakpoint::ETriggerType::FragmentWrite:
+	case UE::Mass::Debug::FBreakpoint::ETriggerType::FragmentAdd:
+	case UE::Mass::Debug::FBreakpoint::ETriggerType::FragmentRemove:
+	case UE::Mass::Debug::FBreakpoint::ETriggerType::TagAdd:
+	case UE::Mass::Debug::FBreakpoint::ETriggerType::TagRemove:
+	{
+		if (!JsonObject->TryGetStringField(TEXT("FragmentName"), TriggerName))
+		{
+			return false;
+		}
+		break;
+	}
+	default:
+		break;
+	}
+
+	{
+		FString FilterTypeStr;
+		if (JsonObject->TryGetStringField(TEXT("FilterType"), FilterTypeStr))
+		{
+			UE::Mass::Debug::FBreakpoint::StringToFilterType(FilterTypeStr, BreakpointInstance.FilterType);
+		}
+	}
+
+	FilterQuery.Reset();
+
+	switch (BreakpointInstance.FilterType)
+	{
+	case UE::Mass::Debug::FBreakpoint::EFilterType::SpecificEntity:
+	{
+		int32 EntityIndex = JsonObject->GetIntegerField(TEXT("EntityIndex"));
+		int32 EntitySerialNumber = JsonObject->GetIntegerField(TEXT("EntitySerialNumber"));
+		BreakpointInstance.Filter.Set<FMassEntityHandle>(FMassEntityHandle(EntityIndex, EntitySerialNumber));
+		break;
+	}
+	case UE::Mass::Debug::FBreakpoint::EFilterType::Query:
+	{
+		if (!JsonObject->TryGetStringField(TEXT("FilterQueryName"), FilterName))
+		{
+			return false;
+		}
+		break;
+	}
+	default:
+		break;
+	}
+
+	return true;
+}
+
+bool FMassDebuggerBreakpointData::DeserializeFromJsonString(const FString& JsonString)
+{
+	TSharedPtr<FJsonObject> JsonObject;
+	TSharedRef<TJsonReader<>> Reader = TJsonReaderFactory<>::Create(JsonString);
+
+	if (!FJsonSerializer::Deserialize(Reader, JsonObject) || !JsonObject.IsValid())
+	{
+		return false;
+	}
+
+	return DeserializeFromJson(JsonObject);
+}
+
+
 float FMassDebuggerModel::MinDistanceToSelectedArchetypes(const TSharedPtr<FMassDebuggerArchetypeData>& InArchetypeData) const
 {
 	float MinDistance = MAX_flt;
@@ -1002,6 +1474,40 @@ void FMassDebuggerModel::OnEntitySelected(const FMassEntityManager& EntityManage
 	}
 }
 
+void FMassDebuggerModel::OnBreakpointsChanged()
+{
+	if (IsStale())
+	{
+		return;
+	}
+	if (TSharedPtr<const FMassEntityManager> EntityManager = Environment->GetEntityManager())
+	{
+		for (int i = CachedBreakpoints.Num() - 1; i >= 0; i--)
+		{
+			TSharedPtr<FMassDebuggerBreakpointData>& BreakData = CachedBreakpoints[i];
+			UE::Mass::Debug::FBreakpoint* EngineBreakpoint = FMassDebugger::FindBreakpoint(*EntityManager, BreakData->BreakpointInstance.Handle);
+			if (EngineBreakpoint)
+			{
+				BreakData->BreakpointInstance = *EngineBreakpoint;
+			}
+		}
+		TArray<UE::Mass::Debug::FBreakpoint>& Breakpoints = FMassDebugger::GetBreakpoints(*EntityManager);
+		for (UE::Mass::Debug::FBreakpoint& EngineBreakpoint : Breakpoints)
+		{
+			bool Exists = Algo::AnyOf(CachedBreakpoints, [InId = EngineBreakpoint.Handle](const TSharedPtr<FMassDebuggerBreakpointData>& BpData)
+			{
+				return BpData->BreakpointInstance.Handle == InId;
+			});
+			if (!Exists)
+			{
+				TSharedPtr<FMassDebuggerBreakpointData> NewBp = MakeShared<FMassDebuggerBreakpointData>(EngineBreakpoint, *this);
+				CachedBreakpoints.Add(NewBp);
+			}
+		}
+		OnEditorBreakpointsChangedDelegate.Broadcast();
+	}
+}
+
 void FMassDebuggerModel::SelectFragment(FName InFragementName)
 {
 	SelectedFragmentName = InFragementName;
@@ -1013,4 +1519,410 @@ FName FMassDebuggerModel::GetSelectedFragment()
 	return SelectedFragmentName;
 }
 
+void FMassDebuggerModel::LoadQueriesFromDisk()
+{
+	const FString SavedDir = FPaths::ProjectSavedDir() / TEXT("MassDebugger");
+	const FString FullPath = SavedDir / TEXT("Queries.json");
+
+	QueryList.Empty();
+
+	if (FPaths::FileExists(FullPath))
+	{
+		FString FileContent;
+		if (FFileHelper::LoadFileToString(FileContent, *FullPath))
+		{
+			TSharedPtr<FJsonObject> RootObj;
+			TSharedRef<TJsonReader<>> Reader = TJsonReaderFactory<>::Create(FileContent);
+			if (FJsonSerializer::Deserialize(Reader, RootObj) && RootObj.IsValid())
+			{
+				const TArray<TSharedPtr<FJsonValue>>* JsonQueries = nullptr;
+				if (RootObj->TryGetArrayField(TEXT("Queries"), JsonQueries))
+				{
+					for (const auto& Val : *JsonQueries)
+					{
+						if (TSharedPtr<FJsonObject> Obj = Val->AsObject())
+						{
+							TSharedPtr<UE::MassDebugger::FEditableQuery> NewQ = MakeShared<UE::MassDebugger::FEditableQuery>();
+							if (NewQ->DeserializeFromJson(Obj))
+							{
+								QueryNames.Add(MakeShared<FString>(NewQ->Name));
+								QueryMap.Add(NewQ->Name, NewQ);
+								QueryList.Add(MoveTemp(NewQ));
+							}
+						}
+					}
+				}
+			}
+		}
+	}
+	RefreshQueries();
+}
+
+void FMassDebuggerModel::RefreshQueries()
+{
+	// this should never happen, but clean up here to be safe and avoid having to check below
+	for (int32 Index = QueryList.Num() - 1; Index >= 0; --Index)
+	{
+		if (!QueryList[Index].IsValid())
+		{
+			QueryList.RemoveAt(Index);
+		}
+	}
+
+	// Clean up duplicate names:
+	const int32 Count = QueryList.Num();
+
+	TArray<FString> OriginalNames;
+	OriginalNames.Reserve(Count);
+	for (const TSharedPtr<UE::MassDebugger::FEditableQuery>& QueryPtr : QueryList)
+	{
+		OriginalNames.Add(QueryPtr.IsValid() ? QueryPtr->Name : FString());
+	}
+
+	TArray<bool> Collides;
+	Collides.Init(false, Count);
+	TMap<FString, int32> MaxSuffix;
+
+	for (int32 i = 0; i < Count; ++i)
+	{
+		const FString& Name = OriginalNames[i];
+		FString Base = Name;
+		int32 UnderscoreIdx;
+		if (Name.FindLastChar(TEXT('_'), UnderscoreIdx))
+		{
+			const FString Sfx = Name.Mid(UnderscoreIdx + 1);
+			if (Sfx.IsNumeric())
+			{
+				Base = Name.Left(UnderscoreIdx);
+				MaxSuffix.FindOrAdd(Base) = FMath::Max(MaxSuffix.FindOrAdd(Base), FCString::Atoi(*Sfx));
+			}
+		}
+		else
+		{
+			MaxSuffix.FindOrAdd(Base);
+		}
+
+		for (int32 j = 0; j < i; ++j)
+		{
+			if (OriginalNames[j].Compare(Name, ESearchCase::IgnoreCase) == 0)
+			{
+				Collides[i] = true;
+				break;
+			}
+		}
+	}
+
+	for (int32 i = 0; i < Count; ++i)
+	{
+		if (!Collides[i] || !QueryList[i].IsValid())
+		{
+			QueryList[i]->Name = OriginalNames[i];
+			continue;
+		}
+
+		const FString& Name = OriginalNames[i];
+		FString Base = Name;
+		int32 UnderscoreIdx;
+		if (Name.FindLastChar(TEXT('_'), UnderscoreIdx))
+		{
+			const FString Sfx = Name.Mid(UnderscoreIdx + 1);
+			if (Sfx.IsNumeric())
+			{
+				Base = Name.Left(UnderscoreIdx);
+			}
+		}
+		int32& Next = MaxSuffix.FindOrAdd(Base);
+		const FString NewName = FString::Printf(TEXT("%s_%d"), *Base, ++Next);
+		QueryList[i]->Name = NewName;
+	}
+
+	// sync up query name list
+	// This could be done more simply by just rebuilding the name list, but we want
+	// to avoid reallocating string ptrs to play nicely with combo box selections.
+	TArray<TSharedPtr<FString>> NewQueryNames;
+	NewQueryNames.Reserve(QueryList.Num());
+
+	for (const TSharedPtr<UE::MassDebugger::FEditableQuery>& QueryPtr : QueryList)
+	{
+		const FString& DesiredName = QueryPtr->Name;
+
+		// ordering might not be consistent so look through the entire list
+		TSharedPtr<FString> ExistingNamePtr;
+		for (const TSharedPtr<FString>& NamePtr : QueryNames)
+		{
+			if (NamePtr.IsValid() && *NamePtr == DesiredName)
+			{
+				ExistingNamePtr = NamePtr;
+				break;
+			}
+		}
+
+		if (ExistingNamePtr.IsValid())
+		{
+			NewQueryNames.Add(ExistingNamePtr);
+		}
+		else
+		{
+			NewQueryNames.Add(MakeShared<FString>(DesiredName));
+		}
+	}
+
+	QueryNames = MoveTemp(NewQueryNames);
+
+	TSet<FString> ValidNames;
+	ValidNames.Reserve(QueryList.Num());
+	for (const TSharedPtr<UE::MassDebugger::FEditableQuery>& QueryPtr : QueryList)
+	{
+		ValidNames.Add(QueryPtr->Name);
+	}
+
+	for (auto It = QueryMap.CreateIterator(); It; ++It)
+	{
+		if (!ValidNames.Contains(It->Key))
+		{
+			It.RemoveCurrent();
+		}
+	}
+
+	for (const TSharedPtr<UE::MassDebugger::FEditableQuery>& QueryPtr : QueryList)
+	{
+		QueryMap.Add(QueryPtr->Name, QueryPtr);
+	}
+
+	OnQueriesChangedDelegate.Broadcast();
+}
+
+namespace UE::MassDebugger
+{
+	void FEditableQuery::InitFromEntityQuery(const FMassEntityQuery& InQuery, FMassDebuggerModel& DebuggerModel)
+	{
+		FragmentRequirements.Reset();
+
+		for (const FMassFragmentRequirementDescription& Desc : InQuery.GetFragmentRequirements())
+		{
+			TSharedPtr<FFragmentEntry> Entry = MakeShared<FFragmentEntry>();
+			Entry->StructType = Desc.StructType;
+			Entry->AccessMode = Desc.AccessMode;
+			Entry->Presence = Desc.Presence;
+			FragmentRequirements.Add(Entry);
+		}
+
+		const FMassTagBitSet& AllTags = InQuery.GetRequiredAllTags();
+		const FMassTagBitSet& AnyTags = InQuery.GetRequiredAnyTags();
+		const FMassTagBitSet& OptionalTags = InQuery.GetRequiredOptionalTags();
+		const FMassTagBitSet& NoneTags = InQuery.GetRequiredNoneTags();
+
+		for (const TSharedPtr<FMassDebuggerFragmentData>& TagData : DebuggerModel.CachedTagData)
+		{
+			const UScriptStruct& StructType = *TagData->Fragment.Get();
+			EMassFragmentPresence Presence = EMassFragmentPresence::All;
+
+			if (AllTags.Contains(StructType))
+			{
+				Presence = EMassFragmentPresence::All;
+			}
+			else if (AnyTags.Contains(StructType))
+			{
+				Presence = EMassFragmentPresence::Any;
+			}
+			else if (OptionalTags.Contains(StructType))
+			{
+				Presence = EMassFragmentPresence::Optional;
+			}
+			else if (NoneTags.Contains(StructType))
+			{
+				Presence = EMassFragmentPresence::None;
+			}
+			else
+			{
+				continue;
+			}
+
+			TSharedPtr<FFragmentEntry> Entry = MakeShared<FFragmentEntry>();
+			Entry->StructType = &StructType;
+			Entry->Presence = Presence;
+			TagRequirements.Add(Entry);
+		}
+	}
+
+	FMassEntityQuery FEditableQuery::BuildEntityQuery(const TSharedRef<FMassEntityManager> EntityManager)
+	{
+		FMassEntityQuery Query;
+
+		Query.Initialize(EntityManager);
+
+		TSet<const UScriptStruct*> UsedFragments;
+		for (TSharedPtr<FFragmentEntry>& Fragment : FragmentRequirements)
+		{
+			const UScriptStruct* RawStruct = Fragment->StructType.Get();
+			if (ensure(RawStruct) && !UsedFragments.Contains(RawStruct))
+			{
+				UsedFragments.Add(RawStruct);
+				Query.AddRequirement(RawStruct, Fragment->AccessMode, Fragment->Presence);
+			}
+		}
+
+		UsedFragments.Reset();
+		for (TSharedPtr<FFragmentEntry>& TagEntry : TagRequirements)
+		{
+			const UScriptStruct* RawStruct = TagEntry->StructType.Get();
+			if (ensure(RawStruct) && !UsedFragments.Contains(RawStruct))
+			{
+				UsedFragments.Add(RawStruct);
+				Query.AddTagRequirement(*RawStruct, TagEntry->Presence);
+			}
+		}
+		return Query;
+	}
+
+	TSharedPtr<FJsonObject> FFragmentEntry::SerializeToJson() const
+	{
+		if (!StructType.IsValid())
+		{
+			return nullptr;
+		}
+		TSharedPtr<FJsonObject> JsonObject = MakeShared<FJsonObject>();
+
+		JsonObject->SetStringField(TEXT("StructType"), StructType->GetPathName());
+		JsonObject->SetStringField(TEXT("AccessMode"), StaticEnum<EMassFragmentAccess>()->GetNameStringByValue((int64)AccessMode));
+		JsonObject->SetStringField(TEXT("Presence"), StaticEnum<EMassFragmentPresence>()->GetNameStringByValue((int64)Presence));
+
+		return JsonObject;
+	}
+
+	bool FFragmentEntry::DeserializeFromJson(const TSharedPtr<FJsonObject>& JsonObject)
+	{
+		if (!JsonObject.IsValid())
+		{
+			return false;
+		}
+		const FString StructPath = JsonObject->GetStringField(TEXT("StructType"));
+		if (!StructPath.IsEmpty())
+		{
+			UScriptStruct* Found = FindObject<UScriptStruct>(nullptr, *StructPath);
+			StructType = Found;
+		}
+		else
+		{
+			StructType = nullptr;
+		}
+
+		const FString AccessStr = JsonObject->GetStringField(TEXT("AccessMode"));
+		UEnum* AccessEnum = StaticEnum<EMassFragmentAccess>();
+		int64 AccessVal = AccessEnum->GetValueByNameString(AccessStr);
+		AccessMode = (EMassFragmentAccess)AccessVal;
+
+		const FString PresenceStr = JsonObject->GetStringField(TEXT("Presence"));
+		UEnum* PresenceEnum = StaticEnum<EMassFragmentPresence>();
+		int64 PresenceVal = PresenceEnum->GetValueByNameString(PresenceStr);
+		Presence = (EMassFragmentPresence)PresenceVal;
+
+		return true;
+	}
+
+	bool FFragmentEntry::DeserializeFromJsonString(const FString& JsonString)
+	{
+		TSharedPtr<FJsonObject> JsonObject;
+		TSharedRef<TJsonReader<>> Reader = TJsonReaderFactory<>::Create(JsonString);
+
+		if (!FJsonSerializer::Deserialize(Reader, JsonObject) || !JsonObject.IsValid())
+		{
+			return false;
+		}
+
+		return DeserializeFromJson(JsonObject);
+	}
+
+	TSharedPtr<FJsonObject> FEditableQuery::SerializeToJson() const
+	{
+		TSharedPtr<FJsonObject> JsonObject = MakeShared<FJsonObject>();
+		JsonObject->SetStringField(TEXT("Name"), Name);
+
+		TArray<TSharedPtr<FJsonValue>> JsonFragments;
+		JsonFragments.Reserve(FragmentRequirements.Num());
+		for (auto& Entry : FragmentRequirements)
+		{
+			JsonFragments.Add(MakeShared<FJsonValueObject>(Entry->SerializeToJson()));
+		}
+		JsonObject->SetArrayField(TEXT("FragmentRequirements"), MoveTemp(JsonFragments));
+
+		TArray<TSharedPtr<FJsonValue>> JsonTags;
+		JsonTags.Reserve(TagRequirements.Num());
+		for (auto& Entry : TagRequirements)
+		{
+			JsonTags.Add(MakeShared<FJsonValueObject>(Entry->SerializeToJson()));
+		}
+		JsonObject->SetArrayField(TEXT("TagRequirements"), MoveTemp(JsonTags));
+
+		return JsonObject;
+	}
+
+	bool FEditableQuery::DeserializeFromJson(const TSharedPtr<FJsonObject>& JsonObject)
+	{
+		if (!JsonObject.IsValid())
+		{
+			return false;
+		}
+
+		if (JsonObject->HasField(TEXT("Name")))
+		{
+			Name = JsonObject->GetStringField(TEXT("Name"));
+		}
+
+		FragmentRequirements.Reset();
+		const TArray<TSharedPtr<FJsonValue>>* JsonFragments = nullptr;
+		if (JsonObject->TryGetArrayField(TEXT("FragmentRequirements"), JsonFragments))
+		{
+			for (const TSharedPtr<FJsonValue>& Val : *JsonFragments)
+			{
+				if (TSharedPtr<FJsonObject> Obj = Val->AsObject())
+				{
+					TSharedPtr<FFragmentEntry> Entry = MakeShared<FFragmentEntry>();
+					if (!Entry->DeserializeFromJson(Obj))
+					{
+						return false;
+					}
+					Entry->bShowAccessMode = true;
+					FragmentRequirements.Add(MoveTemp(Entry));
+				}
+			}
+		}
+
+		TagRequirements.Reset();
+		const TArray<TSharedPtr<FJsonValue>>* JsonTags = nullptr;
+		if (JsonObject->TryGetArrayField(TEXT("TagRequirements"), JsonTags))
+		{
+			for (const TSharedPtr<FJsonValue>& Val : *JsonTags)
+			{
+				if (TSharedPtr<FJsonObject> Obj = Val->AsObject())
+				{
+					TSharedPtr<FFragmentEntry> Entry = MakeShared<FFragmentEntry>();
+					if (!Entry->DeserializeFromJson(Obj))
+					{
+						return false;
+					}
+					Entry->bShowAccessMode = false;
+					TagRequirements.Add(MoveTemp(Entry));
+				}
+			}
+		}
+
+		return true;
+	}
+
+	bool FEditableQuery::DeserializeFromJsonString(const FString& JsonString)
+	{
+		TSharedPtr<FJsonObject> JsonObject;
+		TSharedRef<TJsonReader<>> Reader = TJsonReaderFactory<>::Create(JsonString);
+
+		if (FJsonSerializer::Deserialize(Reader, JsonObject) && JsonObject.IsValid())
+		{
+			return DeserializeFromJson(JsonObject);
+		}
+
+		return false;
+	}
+} // namespace UE::MassDebugger
+
+
 #undef LOCTEXT_NAMESPACE
diff --git a/Engine/Source/Editor/MassEntityDebugger/Private/MassDebuggerModel.h b/Engine/Source/Editor/MassEntityDebugger/Private/MassDebuggerModel.h
--- a/Engine/Source/Editor/MassEntityDebugger/Private/MassDebuggerModel.h
+++ b/Engine/Source/Editor/MassEntityDebugger/Private/MassDebuggerModel.h
@@ -10,6 +10,7 @@
 #include "MassDebugger.h"
 
 
+class FJsonObject;
 class UMassProcessor;
 struct FMassArchetypeHandle;
 struct FMassDebuggerModel;
@@ -23,7 +24,75 @@ class SMassDebugger;
 namespace UE::MassDebugger
 {
 	class SQueryEditorView;
-}
+
+	/** A serializable wrapper for a fragment type and access method used for editable queries */
+	struct FFragmentEntry
+	{
+		/** Convert this struct to a Json object */
+		TSharedPtr<FJsonObject> SerializeToJson() const;
+
+		/** Populate this struct from a Json object created with the above Serialize method */
+		bool DeserializeFromJson(const TSharedPtr<FJsonObject>& JsonObject);
+
+		/** Parse Json string then serialize from that (paste) */
+		bool DeserializeFromJsonString(const FString& JsonString);
+
+		/** The current struct type (selected in the editor combo box) */
+		TWeakObjectPtr<const UScriptStruct> StructType;
+
+		/** Access mode for this fragment entry in a query (not used for tags) */
+		EMassFragmentAccess AccessMode = EMassFragmentAccess::ReadOnly;
+
+		/** Presense requred for this fragment in a query */
+		EMassFragmentPresence Presence = EMassFragmentPresence::All;
+
+		/** Used to hide access mode when unused (e.g. with tags) */
+		bool bShowAccessMode = true;
+	};
+
+	/** A serializable and editable representation of a FMassEntityQuery */
+	struct FEditableQuery
+	{
+		/**
+		 * Convert this entire object and it FragmentEntry arrays to Json
+		 * Used to save user queries created in the debugger UI
+		 */
+		TSharedPtr<FJsonObject> SerializeToJson() const;
+
+		/** 
+		 * Deserialize Json object hierarchy and populate this struct
+		 * Used to load saved user queries when opening the debugger
+		 */
+		bool DeserializeFromJson(const TSharedPtr<FJsonObject>& JsonObject);
+
+		/**
+		 * Parse the Json string then deserialize
+		 * Used to paste queries in the debugger
+		 */
+		bool DeserializeFromJsonString(const FString& JsonString);
+
+		/**
+		 * Create an editable query from a runtime query struct
+		 * Used to view/edit and use runtime queries (e.g. on a Processor) in the editor
+		 */
+		void InitFromEntityQuery(const FMassEntityQuery& InQuery, FMassDebuggerModel& DebuggerModel);
+
+		/** The visible/editable name of the query */
+		FString Name;
+
+		/** Editable representation of the fragment requirements for the query */
+		TArray<TSharedPtr<FFragmentEntry>> FragmentRequirements;
+
+		/** Editable representation of the tag requirements for the query */
+		TArray<TSharedPtr<FFragmentEntry>> TagRequirements;
+
+		/**
+		 * Create a usable runtime query from this editable query
+		 * Used by the debugger to manually run entity queries or for breakpoints
+		 */
+		FMassEntityQuery BuildEntityQuery(const TSharedRef<FMassEntityManager> EntityManager);
+	};
+} // namespace UE::MassDebugger
 
 enum class EMassDebuggerSelectionMode : uint8
 {
@@ -153,6 +222,26 @@ struct FMassDebuggerFragmentData
 	int32 NumEntities = 0;
 };
 
+struct FMassDebuggerBreakpointData
+{
+	FMassDebuggerBreakpointData() = default;
+	FMassDebuggerBreakpointData(const UE::Mass::Debug::FBreakpoint& InBreakpoint, FMassDebuggerModel& Model);
+	void ReconcileDataFromNames(FMassDebuggerModel& Model);
+	void ReconcileDataFromEngineBreakpoint(FMassDebuggerModel& Model);
+	void ApplyToEngine(FMassDebuggerModel& Model, bool bRefreshEngineBreakpoints = true);
+
+	TSharedPtr<FJsonObject> SerializeToJson() const;
+	bool DeserializeFromJson(const TSharedPtr<FJsonObject>& JsonObject);
+	bool DeserializeFromJsonString(const FString& JsonString);
+
+	UE::Mass::Debug::FBreakpoint BreakpointInstance;
+	FString TriggerName;
+	FString FilterName;
+	TSharedPtr<FMassDebuggerFragmentData> TriggerFragment;
+	TSharedPtr<FMassDebuggerProcessorData> TriggerProcessor;
+	TSharedPtr<UE::MassDebugger::FEditableQuery> FilterQuery;
+};
+
 struct FMassDebuggerProcessingGraphNode
 {
 	explicit FMassDebuggerProcessingGraphNode(const TSharedPtr<FMassDebuggerProcessorData>& InProcessorData, const UMassCompositeProcessor::FDependencyNode& InProcessorNode = UMassCompositeProcessor::FDependencyNode());
@@ -200,6 +289,8 @@ struct FMassDebuggerModel
 	DECLARE_MULTICAST_DELEGATE_TwoParams(FOnProcessorsSelected, TConstArrayView<TSharedPtr<FMassDebuggerProcessorData>>, ESelectInfo::Type);
 	DECLARE_MULTICAST_DELEGATE_TwoParams(FOnArchetypesSelected, TConstArrayView<TSharedPtr<FMassDebuggerArchetypeData>>, ESelectInfo::Type);
 	DECLARE_MULTICAST_DELEGATE_OneParam(FOnFragmentSelected, FName FragmentName);
+	DECLARE_MULTICAST_DELEGATE(FOnQueriesChanged);
+	DECLARE_MULTICAST_DELEGATE(FOnBreakpointsChanged);
 
 	FMassDebuggerModel();
 	~FMassDebuggerModel();
@@ -208,6 +299,8 @@ struct FMassDebuggerModel
 
 	void RefreshAll();
 
+	void LoadQueriesFromDisk();
+
 	void SelectProcessor(TSharedPtr<FMassDebuggerProcessorData>& Processor);
 	void SelectProcessors(TArrayView<TSharedPtr<FMassDebuggerProcessorData>> Processors, ESelectInfo::Type SelectInfo);
 	void ClearProcessorSelection();
@@ -228,7 +321,18 @@ struct FMassDebuggerModel
 		const TConstArrayView<TWeakObjectPtr<const UScriptStruct>>& FragmentTypes
 	);
 	void CacheFragmentData();
+	TSharedPtr<FMassDebuggerFragmentData> FindFragmentData(const UScriptStruct* FragmentType);
+
 	void CacheTagData();
+	
+	void LoadBreakpointsFromDisk();
+	void SaveBreakpointsToDisk();
+	void CreateBreakpoint();
+	void CreateBreakpointFromString(const FString& InString);
+	void ApplyBreakpointsToCurrentEnvironment();
+	void ReconcileAllBreakpoints();
+	void RemoveAllBreakpoints();
+	void RemoveBreakpoint(UE::Mass::Debug::FBreakpointHandle Handle);
 
 	float MinDistanceToSelectedArchetypes(const TSharedPtr<FMassDebuggerArchetypeData>& InArchetypeData) const;
 
@@ -251,6 +355,7 @@ struct FMassDebuggerModel
 
 	void RegisterQueryEditor(TSharedRef<UE::MassDebugger::SQueryEditorView> InQueryEditorView);
 	void ShowQueryInEditor(const FMassEntityQuery& InQuery, const FString& InQueryName);
+	void RefreshQueries();
 
 protected:
 	void StoreArchetypes(const FMassEntityManager& EntityManager, TMap<FMassArchetypeHandle, TSharedPtr<FMassDebuggerArchetypeData>>& OutTransientArchetypesMap);
@@ -259,6 +364,7 @@ protected:
 	void ResetSelectedProcessors();
 
 	void OnEntitySelected(const FMassEntityManager& EntityManager, const FMassEntityHandle EntityHandle);
+	void OnBreakpointsChanged();
 
 public:
 	TWeakPtr<SMassDebugger> DebuggerWindow;
@@ -267,6 +373,8 @@ public:
 	FOnProcessorsSelected OnProcessorsSelectedDelegate;
 	FOnArchetypesSelected OnArchetypesSelectedDelegate;
 	FOnFragmentSelected OnFragmentSelectedDelegate;
+	FOnBreakpointsChanged OnEditorBreakpointsChangedDelegate;
+	FOnQueriesChanged OnQueriesChangedDelegate;
 
 	EMassDebuggerSelectionMode SelectionMode = EMassDebuggerSelectionMode::None;
 
@@ -281,6 +389,7 @@ public:
 		TArray<TSharedPtr<FMassDebuggerProcessorData>> Container;
 	};
 
+	TArray<TSharedPtr<UE::MassDebugger::FEditableQuery>> QueryList;
 	TArray<TSharedPtr<FMassDebuggerFragmentData>> CachedTagData;
 	TArray<TSharedPtr<FMassDebuggerFragmentData>> CachedFragmentData;
 	TArray<TSharedPtr<FProcessorCollection>> CachedProcessorCollections;
@@ -289,6 +398,17 @@ public:
 	TArray<TSharedPtr<FMassDebuggerArchetypeData>> CachedArchetypeRepresentatives;
 	TArray<TSharedPtr<FMassDebuggerArchetypeData>> SelectedArchetypes;
 	TArray<TSharedPtr<FMassDebuggerProcessingGraph>> CachedProcessingGraphs;
+	TArray<TSharedPtr<FMassDebuggerBreakpointData>> CachedBreakpoints;
+
+	/** Name lists and maps for easy lookup from widgets */
+	TArray<TSharedPtr<FString>> ProcessorNames;
+	TMap<FString, TWeakPtr<FMassDebuggerProcessorData>> ProcessorMap;
+	TArray<TSharedPtr<FString>> FragmentNames;
+	TMap<FString, TWeakPtr<FMassDebuggerFragmentData>> FragmentMap;
+	TArray<TSharedPtr<FString>> TagNames;
+	TMap<FString, TWeakPtr<FMassDebuggerFragmentData>> TagMap;
+	TArray<TSharedPtr<FString>> QueryNames;
+	TMap<FString, TWeakPtr<UE::MassDebugger::FEditableQuery>> QueryMap;
 
 	TMap<FMassArchetypeHandle, TSharedPtr<FMassDebuggerArchetypeData>> HandleToArchetypeMap;
 
@@ -297,6 +417,7 @@ public:
 	FString EnvironmentDisplayName;
 
 	FDelegateHandle OnEntitySelectedHandle;
+	FDelegateHandle OnBreakpointsChangedHandle;
 
 	void SelectFragment(FName InFragementName);
 	FName GetSelectedFragment();
diff --git a/Engine/Source/Editor/MassEntityDebugger/Private/SMassBreakpointsView.cpp b/Engine/Source/Editor/MassEntityDebugger/Private/SMassBreakpointsView.cpp
--- a/Engine/Source/Editor/MassEntityDebugger/Private/SMassBreakpointsView.cpp
+++ b/Engine/Source/Editor/MassEntityDebugger/Private/SMassBreakpointsView.cpp
@@ -1,195 +1,615 @@
 // Copyright Epic Games, Inc. All Rights Reserved.
 
 #include "SMassBreakpointsView.h"
-#include "Widgets/Input/SButton.h"
-#include "Widgets/Layout/SScrollBox.h"
-#include "Widgets/Views/SListView.h"
-#include "Widgets/SBoxPanel.h"
 #include "MassDebugger.h"
 #include "MassDebuggerModel.h"
+#include "MassProcessor.h"
+#include "HAL/PlatformApplicationMisc.h"
+#include "Widgets/Input/SButton.h"
+#include "Widgets/Text/STextBlock.h"
+#include "Widgets/SBoxPanel.h"
+#include "SSearchableComboBox.h"
+#include "Serialization/JsonSerializer.h"
+#include "Serialization/JsonReader.h"
+#include "Serialization/JsonWriter.h"
+#include "Dom/JsonObject.h"
 
 #define LOCTEXT_NAMESPACE "SMassDebugger"
 
-
-//----------------------------------------------------------------------//
-// SMassBreakpointsView
-//----------------------------------------------------------------------//
-void SMassBreakpointsView::Construct(const FArguments& InArgs, TSharedRef<FMassDebuggerModel> InDebuggerModel)
+void SMassBreakpointsView::Construct(const FArguments& InArgs, TSharedRef<FMassDebuggerModel> InModel)
 {
+	Initialize(InModel);
+
 #if WITH_MASSENTITY_DEBUG
-	Initialize(InDebuggerModel);
+	InModel->OnEditorBreakpointsChangedDelegate.AddSP(this, &SMassBreakpointsView::RefreshBreakpoints);
+	InModel->OnQueriesChangedDelegate.AddSP(this, &SMassBreakpointsView::RefreshBreakpoints);
+
+	TriggerTypeOptions = {
+		MakeShared<FString>(UE::Mass::Debug::FBreakpoint::TriggerTypeToString(
+			UE::Mass::Debug::FBreakpoint::ETriggerType::None)),
+		MakeShared<FString>(UE::Mass::Debug::FBreakpoint::TriggerTypeToString(
+			UE::Mass::Debug::FBreakpoint::ETriggerType::ProcessorExecute)),
+		MakeShared<FString>(UE::Mass::Debug::FBreakpoint::TriggerTypeToString(
+			UE::Mass::Debug::FBreakpoint::ETriggerType::FragmentWrite)),
+		MakeShared<FString>(UE::Mass::Debug::FBreakpoint::TriggerTypeToString(
+			UE::Mass::Debug::FBreakpoint::ETriggerType::EntityCreate)),
+		MakeShared<FString>(UE::Mass::Debug::FBreakpoint::TriggerTypeToString(
+			UE::Mass::Debug::FBreakpoint::ETriggerType::EntityDestroy)),
+		MakeShared<FString>(UE::Mass::Debug::FBreakpoint::TriggerTypeToString(
+			UE::Mass::Debug::FBreakpoint::ETriggerType::FragmentAdd)),
+		MakeShared<FString>(UE::Mass::Debug::FBreakpoint::TriggerTypeToString(
+			UE::Mass::Debug::FBreakpoint::ETriggerType::FragmentRemove)),
+		MakeShared<FString>(UE::Mass::Debug::FBreakpoint::TriggerTypeToString(
+			UE::Mass::Debug::FBreakpoint::ETriggerType::TagAdd)),
+		MakeShared<FString>(UE::Mass::Debug::FBreakpoint::TriggerTypeToString(
+			UE::Mass::Debug::FBreakpoint::ETriggerType::TagRemove))
+	};
+
+	check(TriggerTypeOptions.Num() == static_cast<int32>(UE::Mass::Debug::FBreakpoint::ETriggerType::MAX));
+
+	for (TSharedPtr<FString>& TypeStringPtr : TriggerTypeOptions)
+	{
+		UE::Mass::Debug::FBreakpoint::ETriggerType Trigger;
+		if (TypeStringPtr && UE::Mass::Debug::FBreakpoint::StringToTriggerType(*TypeStringPtr, Trigger))
+		{
+			TriggerMap.Add(Trigger, TypeStringPtr);
+		}
+	}
 
-	FMassDebugger::OnBreakpointsChangedDelegate.AddSP(this, &SMassBreakpointsView::RefreshBreakpoints);
+	FilterTypeOptions = {
+		MakeShared<FString>(UE::Mass::Debug::FBreakpoint::FilterTypeToString(
+			UE::Mass::Debug::FBreakpoint::EFilterType::None)),
+		MakeShared<FString>(UE::Mass::Debug::FBreakpoint::FilterTypeToString(
+			UE::Mass::Debug::FBreakpoint::EFilterType::SpecificEntity)),
+		MakeShared<FString>(UE::Mass::Debug::FBreakpoint::FilterTypeToString(
+			UE::Mass::Debug::FBreakpoint::EFilterType::SelectedEntity)),
+		MakeShared<FString>(UE::Mass::Debug::FBreakpoint::FilterTypeToString(
+			UE::Mass::Debug::FBreakpoint::EFilterType::Query))
+	};
 
-	WriteBreakpointsListView = SNew(SListView<TSharedPtr<FBreakpointDisplay>>)
-		.ListItemsSource(&WriteBreakpoints)
-		.OnGenerateRow(this, &SMassBreakpointsView::OnGenerateBreakpointRow);
+	for (TSharedPtr<FString>& TypeStringPtr : FilterTypeOptions)
+	{
+		UE::Mass::Debug::FBreakpoint::EFilterType Filter;
+		if (TypeStringPtr && UE::Mass::Debug::FBreakpoint::StringToFilterType(*TypeStringPtr, Filter))
+		{
+			FilterMap.Add(Filter, TypeStringPtr);
+		}
+	}
+	
+	BreakpointsListView = SNew(SListView<TSharedPtr<FMassDebuggerBreakpointData>>)
+	.ListItemsSource(&InModel->CachedBreakpoints)
+	.OnGenerateRow(this, &SMassBreakpointsView::OnGenerateBreakpointRow)
+	.HeaderRow(
+		SNew(SHeaderRow)
+		+ SHeaderRow::Column("Enabled")
+		.DefaultLabel(FText::GetEmpty())
+		.FixedWidth(128)
+		+ SHeaderRow::Column("TriggerType")
+		.DefaultLabel(LOCTEXT("TriggerTypeColumn", "Trigger Type"))
+		.FillWidth(0.2f)
+		+ SHeaderRow::Column("Trigger")
+		.DefaultLabel(LOCTEXT("TriggerColumn", "Trigger"))
+		.FillWidth(0.2f)
+		+ SHeaderRow::Column("FilterType")
+		.DefaultLabel(LOCTEXT("FilterTypeColumn", "Filter Type"))
+		.FillWidth(0.2f)
+		+ SHeaderRow::Column("Filter")
+		.DefaultLabel(LOCTEXT("FilterColumn", "Filter"))
+		.FillWidth(0.2f)
+		+ SHeaderRow::Column(TEXT("HitCount"))
+		.DefaultLabel(LOCTEXT("HitCount", "Hit Count"))
+		.FillWidth(0.2f)
+	);
 
 	ChildSlot
 	[
 		SNew(SVerticalBox)
-		+ SVerticalBox::Slot()
-		.AutoHeight()
-		.Padding(5)
+    
+		+ SVerticalBox::Slot().AutoHeight().Padding(4)
 		[
-			SNew(SButton)
-			.Text(LOCTEXT("ClearAllBreakpoints", "Clear All Breakpoints"))
-			.OnClicked(this, &SMassBreakpointsView::ClearBreakpointsClicked)
+			SNew(SHorizontalBox)
+			+ SHorizontalBox::Slot().AutoWidth()
+			[
+				SNew(SButton)
+				.ToolTipText(LOCTEXT("NewBreakpointTooltip", "New Breakpoint"))
+				.OnClicked(this, &SMassBreakpointsView::NewBreakpointClicked)
+				[
+					SNew(SImage).Image(FAppStyle::GetBrush("Icons.PlusCircle"))
+				]
+			]
+			+ SHorizontalBox::Slot().AutoWidth()
+			[
+				SNew(SButton)
+				.ToolTipText(LOCTEXT("PasteBreakpointTooltip", "Paste Breakpoint"))
+				.OnClicked(this, &SMassBreakpointsView::PasteBreakpointClicked)
+				[
+					SNew(SImage).Image(FAppStyle::GetBrush("GenericCommands.Paste"))
+				]
+			]
+			+ SHorizontalBox::Slot().AutoWidth().Padding(0,0,4,0)
+			[
+				SNew(SButton)
+				.ToolTipText(LOCTEXT("SaveBreakpointsTooltip", "Save Breakpoints"))
+				.OnClicked(this, &SMassBreakpointsView::SaveBreakpointsClicked)
+				[
+					SNew(SImage).Image(FAppStyle::GetBrush("Icons.SaveModified"))
+				]
+			]
+			+ SHorizontalBox::Slot().AutoWidth().Padding(0,0,4,0)
+			[
+				SNew(SButton)
+				.ToolTipText(LOCTEXT("ClearAllBreakpoints", "Clear All Breakpoints"))
+				.OnClicked(this, &SMassBreakpointsView::ClearBreakpointsClicked)
+				[
+					SNew(SImage).Image(FAppStyle::GetBrush("Icons.Delete"))
+				]
+			]
 		]
-		+ SVerticalBox::Slot()
-		.FillHeight(1.f)
-		.Padding(5)
+		+ SVerticalBox::Slot().FillHeight(1.f).Padding(4)
 		[
-			WriteBreakpointsListView.ToSharedRef()
+			BreakpointsListView.ToSharedRef()
 		]
 	];
 #else
 	ChildSlot
-	[
-		SNew(STextBlock)
-			.Text(LOCTEXT("MassEntityDebuggingNotEnabled", "Mass Entity Debugging Not Enabled for this configuration"))
-	];
+		[
+			SNew(STextBlock).Text(LOCTEXT("DebugNotEnabled", "Mass Entity Debugging Not Enabled"))
+		];
 #endif
 }
 
-void SMassBreakpointsView::OnRefresh()
+void SMassBreakpointsView::RefreshBreakpoints()
 {
-	RefreshBreakpoints();
+#if WITH_MASSENTITY_DEBUG
+	if (BreakpointsListView.IsValid())
+	{
+		BreakpointsListView->RebuildList();
+	}
+#endif
 }
 
-void SMassBreakpointsView::RefreshBreakpoints()
+FReply SMassBreakpointsView::ClearBreakpointsClicked()
 {
 #if WITH_MASSENTITY_DEBUG
-	WriteBreakpoints.Reset();
-	if (DebuggerModel && DebuggerModel->Environment && DebuggerModel->Environment->EntityManager.IsValid())
+	if (DebuggerModel)
 	{
-		FMassDebugger::FEnvironment* Env = FMassDebugger::FindEnvironmentForEntityManager(DebuggerModel->Environment->EntityManager.Pin().ToSharedRef().Get());
-		if (Env == nullptr)
-		{
-			return;
-		}
-
-		TMultiMap<TObjectKey<const UScriptStruct>, FMassEntityHandle>::TIterator FragmentIterator = Env->FragmentWriteBreakpoints.CreateIterator();
-		for (; FragmentIterator; ++FragmentIterator)
-		{
-			TSharedPtr<FBreakpointDisplay> NewDisplay = MakeShared<FBreakpointDisplay>();
-			NewDisplay->WriteFragment = FragmentIterator.Key().ResolveObjectPtr();
-			NewDisplay->Entity = FragmentIterator.Value();
-			if (NewDisplay->WriteFragment.IsValid())
-			{
-				WriteBreakpoints.Add(NewDisplay);
-			}
-		}
-		for (TObjectKey<const UScriptStruct>& FragmentTypeKey : Env->SelectedEntityFragmentWriteBreakpoints)
-		{
-			TSharedPtr<FBreakpointDisplay> NewDisplay = MakeShared<FBreakpointDisplay>();
-			NewDisplay->WriteFragment = FragmentTypeKey.ResolveObjectPtr();
-			NewDisplay->bBreakForSelectedEntity = true;
-			if (NewDisplay->WriteFragment.IsValid())
-			{
-				WriteBreakpoints.Add(NewDisplay);
-			}
-		}
+		DebuggerModel->RemoveAllBreakpoints();
 	}
-	if (WriteBreakpointsListView)
+#endif
+	return FReply::Handled();
+}
+
+FReply SMassBreakpointsView::SaveBreakpointsClicked()
+{
+#if WITH_MASSENTITY_DEBUG
+	if (DebuggerModel)
 	{
-		WriteBreakpointsListView->RebuildList();
+		DebuggerModel->SaveBreakpointsToDisk();
 	}
 #endif
+	return FReply::Handled();
 }
 
-FReply SMassBreakpointsView::HandleAddWriteBreakpointClicked()
+FReply SMassBreakpointsView::PasteBreakpointClicked()
 {
 #if WITH_MASSENTITY_DEBUG
-	FMassDebugger::BreakOnFragmentWriteForSelectedEntity(DebuggerModel->GetSelectedFragment()); 
+	if (DebuggerModel)
+	{
+		FString ClipboardText;
+		FPlatformApplicationMisc::ClipboardPaste(ClipboardText);
+
+		DebuggerModel->CreateBreakpointFromString(ClipboardText);
+	}
 #endif
 	return FReply::Handled();
 }
 
-FReply SMassBreakpointsView::ClearBreakpointsClicked()
+FReply SMassBreakpointsView::NewBreakpointClicked()
 {
 #if WITH_MASSENTITY_DEBUG
-	FMassDebugger::ClearAllBreakpoints();
+	if (DebuggerModel)
+	{
+		DebuggerModel->CreateBreakpoint();
+	}
 #endif
 	return FReply::Handled();
 }
 
-TSharedRef<ITableRow> SMassBreakpointsView::OnGenerateBreakpointRow(TSharedPtr<FBreakpointDisplay> InItem, const TSharedRef<STableViewBase>& OwnerTable)
+TSharedRef<ITableRow> SMassBreakpointsView::OnGenerateBreakpointRow(
+    TSharedPtr<FMassDebuggerBreakpointData> InItem,
+    const TSharedRef<STableViewBase>& OwnerTable)
 {
-	return SNew(STableRow<TSharedPtr<FMassEntityHandle>>, OwnerTable)
-		.Content()
+	return SNew(SMassBreakpointsView::SBreakpointTableRow, OwnerTable, SharedThis(this), InItem.ToSharedRef(), DebuggerModel.ToSharedRef());
+}
+
+void SMassBreakpointsView::SBreakpointTableRow::Construct(const FArguments& InArgs, const TSharedRef<STableViewBase>& InOwnerTableView, TSharedRef<SMassBreakpointsView> InParentView, TSharedRef<FMassDebuggerBreakpointData> InBreakpointData, TSharedRef<FMassDebuggerModel> InModel)
+{
+	BreakpointData = InBreakpointData;
+	DebuggerModel = InModel;
+	ParentView = InParentView;
+	SMultiColumnTableRow<TSharedPtr<FMassDebuggerBreakpointData>>::Construct(SMultiColumnTableRow<TSharedPtr<FMassDebuggerBreakpointData>>::FArguments(), InOwnerTableView);
+}
+
+TSharedRef<SWidget> SMassBreakpointsView::SBreakpointTableRow::GenerateWidgetForColumn(const FName& ColumnId)
+{
+	TSharedPtr<FMassDebuggerBreakpointData> Data = BreakpointData.Pin();
+	TSharedPtr<SMassBreakpointsView> View = ParentView.Pin();
+
+	if (!Data || !View)
+	{
+		return SNullWidget::NullWidget;
+	}
+
+	if (ColumnId == "Enabled")
+	{
+		return SNew(SHorizontalBox)
+		+ SHorizontalBox::Slot()
+		.AutoWidth()
+		.Padding(2)
 		[
-#if WITH_MASSENTITY_DEBUG
-            SNew(SHorizontalBox)
-			+ SHorizontalBox::Slot()
-            .AutoWidth()
-            .Padding(5.0f, 2.0f)
-            [
-                // Add a button to remove the breakpoint
-                SNew(SButton)
-                .OnClicked_Lambda([this, InItem]()
-                {
-					if (DebuggerModel.IsValid() && DebuggerModel->Environment.IsValid())
+			SNew(SCheckBox)
+				.IsChecked_Lambda([WeakItem = BreakpointData]()
+				{
+					if (TSharedPtr<FMassDebuggerBreakpointData> Item = WeakItem.Pin())
 					{
-						const FMassEntityManager* EntityManager = DebuggerModel->Environment->EntityManager.Pin().Get();
-						if (EntityManager)
+						return Item->BreakpointInstance.bEnabled ? ECheckBoxState::Checked : ECheckBoxState::Unchecked;
+					}
+					return ECheckBoxState::Unchecked;
+				})
+				.OnCheckStateChanged_Lambda([WeakItem = BreakpointData, WeakModel = DebuggerModel](ECheckBoxState NewState)
+				{
+					if (TSharedPtr<FMassDebuggerBreakpointData> Item = WeakItem.Pin())
+					{
+						Item->BreakpointInstance.bEnabled = (NewState == ECheckBoxState::Checked);
+						if (TSharedPtr<FMassDebuggerModel> Model = WeakModel.Pin())
+						{
+							Model->ApplyBreakpointsToCurrentEnvironment();
+						}
+					}
+				})
+		]
+		+ SHorizontalBox::Slot()
+		.AutoWidth()
+		.Padding(2)
+		[
+			SNew(SButton)
+				.ToolTipText(LOCTEXT("DeleteBreakpointTooltip", "Delete Breakpoint"))
+				.OnClicked_Lambda([WeakItem = BreakpointData, WeakModel = DebuggerModel]()
+				{
+					if (TSharedPtr<FMassDebuggerBreakpointData> Item = WeakItem.Pin())
+					{
+						if (TSharedPtr<FMassDebuggerModel> Model = WeakModel.Pin())
 						{
-							if (InItem->Processor.IsValid())
-							{
-								FMassDebugger::ClearProcessorBreakpoint(*EntityManager, InItem->Processor.Get(), InItem->Entity);
-							}
-							else if (InItem->WriteFragment.IsValid())
-							{
-								if (InItem->bBreakForSelectedEntity)
-								{
-									FMassDebugger::ToggleFragmentWriteBreakForSelectedEntity(*EntityManager, InItem->WriteFragment.Get());
-								}
-								else
-								{
-									FMassDebugger::ClearFragmentWriteBreak(*EntityManager, InItem->WriteFragment.Get(), InItem->Entity);
-								}
-							}
+							Model->RemoveBreakpoint(Item->BreakpointInstance.Handle);
 						}
 					}
-                    return FReply::Handled();
+					return FReply::Handled();
+				})
+			[
+				SNew(SImage).Image(FAppStyle::GetBrush("Icons.Delete"))
+			]
+		]
+		+ SHorizontalBox::Slot()
+		.AutoWidth()
+		.Padding(2)
+		[
+			SNew(SButton)
+			.ToolTipText(LOCTEXT("CopyBreakpointTooltip", "Copy Breakpoint"))
+			.OnClicked_Lambda([WeakItem = BreakpointData]()
+			{
+				if (TSharedPtr<FMassDebuggerBreakpointData> Item = WeakItem.Pin())
+				{
+					FString JsonString;
+					const TSharedRef<TJsonWriter<>> Writer = TJsonWriterFactory<>::Create(&JsonString);
+					FJsonSerializer::Serialize(Item->SerializeToJson().ToSharedRef(), Writer);
+
+					FPlatformApplicationMisc::ClipboardCopy(*JsonString);
+				}
+				return FReply::Handled();
+			})
+			[
+				SNew(SImage)
+				.Image(FAppStyle::GetBrush("GenericCommands.Copy"))
+			]
+		];
+	}
+    else if (ColumnId == "TriggerType")
+    {
+        return SNew(SComboBox<TSharedPtr<FString>>)
+        .OptionsSource(&View->TriggerTypeOptions)
+        .OnGenerateWidget_Lambda([](TSharedPtr<FString> Str)
+		{
+			return SNew(STextBlock).Text(FText::FromString(*Str));
+		})
+        .OnSelectionChanged_Lambda([WeakData = BreakpointData, WeakModel = DebuggerModel](TSharedPtr<FString> Str, ESelectInfo::Type)
+        {
+            UE::Mass::Debug::FBreakpoint::ETriggerType NewType;
+			TSharedPtr<FMassDebuggerBreakpointData> Data = WeakData.Pin();
+            if (Data && Str && UE::Mass::Debug::FBreakpoint::StringToTriggerType(*Str, NewType))
+            {
+				Data->BreakpointInstance.TriggerType = NewType;
+				if (TSharedPtr<FMassDebuggerModel> Model = WeakModel.Pin())
+				{
+					Model->ApplyBreakpointsToCurrentEnvironment();
+				}
+            }
+        })
+        .InitiallySelectedItem(View->TriggerMap[Data->BreakpointInstance.TriggerType])
+        [
+            SNew(STextBlock)
+			.Text_Lambda([WeakData = BreakpointData]()
+			{
+				if (TSharedPtr<FMassDebuggerBreakpointData> Data = WeakData.Pin())
+				{
+					return FText::FromString(UE::Mass::Debug::FBreakpoint::TriggerTypeToString(Data->BreakpointInstance.TriggerType));
+				}
+				return LOCTEXT("BreakpointInvalid", "Invalid Breakpoint");
+			})
+        ];
+    }
+    else if (ColumnId == "Trigger")
+    {
+        if (Data->BreakpointInstance.TriggerType == UE::Mass::Debug::FBreakpoint::ETriggerType::ProcessorExecute)
+        {
+            return SNew(SSearchableComboBox)
+            .OptionsSource(&DebuggerModel.Pin()->ProcessorNames)
+            .OnGenerateWidget_Lambda([](TSharedPtr<FString> Str)
+			{
+				if (Str)
+				{
+					return SNew(STextBlock)
+						.Text(FText::FromString(*Str));
+				}
+				return SNew(STextBlock)
+					.Text(LOCTEXT("ProcessorNone", "None"));
+			})
+            .OnSelectionChanged_Lambda([WeakData = BreakpointData, WeakModel = DebuggerModel](TSharedPtr<FString> Str, ESelectInfo::Type)
+            {
+				if (TSharedPtr<FMassDebuggerBreakpointData> Data = WeakData.Pin())
+				{
+					if (Str)
+					{
+						Data->TriggerName = *Str;
+					}
+					else
+					{
+						Data->TriggerName.Empty();
+					}
+
+					if (TSharedPtr<FMassDebuggerModel> Model = WeakModel.Pin())
+					{
+						Data->ReconcileDataFromNames(*Model);
+						Model->ApplyBreakpointsToCurrentEnvironment();
+					}
+				}
+            })
+            [
+                SNew(STextBlock).Text_Lambda([WeakData = BreakpointData]()
+                {
+					if (TSharedPtr<FMassDebuggerBreakpointData> Data = WeakData.Pin())
+					{
+						return FText::FromString(Data->TriggerName);
+					}
+                    return LOCTEXT("SelectProcessor","Select Processor");
                 })
-				[
-					SNew(SImage)
-						.Image(FCoreStyle::Get().GetBrush("Icons.Delete"))
-						.ColorAndOpacity(FSlateColor::UseForeground())
-				]
-            ]
-            + SHorizontalBox::Slot()
-            .AutoWidth()
-            .Padding(5.0f, 2.0f)
+            ];
+        }
+        else if (Data->BreakpointInstance.TriggerType == UE::Mass::Debug::FBreakpoint::ETriggerType::FragmentWrite
+			|| Data->BreakpointInstance.TriggerType == UE::Mass::Debug::FBreakpoint::ETriggerType::FragmentAdd
+			|| Data->BreakpointInstance.TriggerType == UE::Mass::Debug::FBreakpoint::ETriggerType::FragmentRemove)
+        {
+            return SNew(SSearchableComboBox)
+            .OptionsSource(&DebuggerModel.Pin()->FragmentNames)
+            .OnGenerateWidget_Lambda([](TSharedPtr<FString> Str)
+			{
+				if (Str)
+				{
+					return SNew(STextBlock)
+						.Text(FText::FromString(*Str));
+				}
+				return SNew(STextBlock)
+					.Text(LOCTEXT("FragmentNone", "None"));
+			})
+            .OnSelectionChanged_Lambda([WeakData = BreakpointData, WeakModel = DebuggerModel](TSharedPtr<FString> Str, ESelectInfo::Type)
+            {
+				if (TSharedPtr<FMassDebuggerBreakpointData> Data = WeakData.Pin())
+				{
+					if (Str)
+					{
+						Data->TriggerName = *Str;
+					}
+					else
+					{
+						Data->TriggerName.Empty();
+					}
+
+					if (TSharedPtr<FMassDebuggerModel> Model = WeakModel.Pin(); !Model->IsStale())
+					{
+						Data->ReconcileDataFromNames(*Model);
+						Data->ApplyToEngine(*Model);
+					}
+				}
+            })
             [
-                // Display WriteFragment or Processor name
-                SNew(STextBlock)
-                .Text(InItem->WriteFragment.Get()
-                    ? FText::FromString(InItem->WriteFragment->GetName()) 
-                    : (InItem->Processor.Get()
-                        ? FText::FromString(InItem->Processor->GetName()) 
-                        : LOCTEXT("UnknownBreakpoint", "Unknown")))
-            ]
-            + SHorizontalBox::Slot()
-            .AutoWidth()
-            .Padding(5.0f, 2.0f)
+				SNew(STextBlock).Text_Lambda([WeakData = BreakpointData]()
+				{
+					if (TSharedPtr<FMassDebuggerBreakpointData> Data = WeakData.Pin())
+					{
+						return FText::FromString(Data->TriggerName);
+					}
+					return LOCTEXT("SelectFragment", "Select Fragment");
+				})
+            ];
+        }
+		else if (Data->BreakpointInstance.TriggerType == UE::Mass::Debug::FBreakpoint::ETriggerType::TagAdd
+			|| Data->BreakpointInstance.TriggerType == UE::Mass::Debug::FBreakpoint::ETriggerType::TagRemove)
+        {
+            return SNew(SSearchableComboBox)
+            .OptionsSource(&DebuggerModel.Pin()->TagNames)
+            .OnGenerateWidget_Lambda([](TSharedPtr<FString> Str)
+			{
+				if (Str)
+				{
+					return SNew(STextBlock)
+						.Text(FText::FromString(*Str));
+				}
+				return SNew(STextBlock)
+					.Text(LOCTEXT("TagNone", "None"));
+			})
+            .OnSelectionChanged_Lambda([WeakData = BreakpointData, WeakModel = DebuggerModel](TSharedPtr<FString> Str, ESelectInfo::Type)
+            {
+				if (TSharedPtr<FMassDebuggerBreakpointData> Data = WeakData.Pin())
+				{
+					if (Str)
+					{
+						Data->TriggerName = *Str;
+					}
+					else
+					{
+						Data->TriggerName.Empty();
+					}
+
+					if (TSharedPtr<FMassDebuggerModel> Model = WeakModel.Pin(); !Model->IsStale())
+					{
+						Data->ReconcileDataFromNames(*Model);
+						Data->ApplyToEngine(*Model);
+					}
+				}
+            })
+            [
+				SNew(STextBlock).Text_Lambda([WeakData = BreakpointData]()
+				{
+					if (TSharedPtr<FMassDebuggerBreakpointData> Data = WeakData.Pin())
+					{
+						return FText::FromString(Data->TriggerName);
+					}
+					return LOCTEXT("SelectTag", "Select Tag");
+				})
+            ];
+        }
+        return SNullWidget::NullWidget;
+    }
+    else if (ColumnId == "FilterType")
+	{
+        return SNew(SComboBox<TSharedPtr<FString>>)
+        .OptionsSource(&View->FilterTypeOptions)
+        .OnGenerateWidget_Lambda([](TSharedPtr<FString> Str)
+		{
+			return SNew(STextBlock).Text(FText::FromString(*Str));
+		})
+        .OnSelectionChanged_Lambda([WeakData = BreakpointData, WeakModel = DebuggerModel](TSharedPtr<FString> Str, ESelectInfo::Type)
+        {
+            
+			if (TSharedPtr<FMassDebuggerBreakpointData> Data = WeakData.Pin())
+			{
+				UE::Mass::Debug::FBreakpoint::EFilterType NewType;
+				if (UE::Mass::Debug::FBreakpoint::StringToFilterType(*Str, NewType))
+				{
+					Data->BreakpointInstance.FilterType = NewType;
+					if (TSharedPtr<FMassDebuggerModel> Model = WeakModel.Pin())
+					{
+						Model->ApplyBreakpointsToCurrentEnvironment();
+					}
+				}
+			}
+        })
+        .InitiallySelectedItem(View->FilterMap[Data->BreakpointInstance.FilterType])
+        [
+            SNew(STextBlock)
+			.Text_Lambda([WeakData = BreakpointData]()
+			{
+				if (TSharedPtr<FMassDebuggerBreakpointData> Data = WeakData.Pin())
+				{
+					return FText::FromString(UE::Mass::Debug::FBreakpoint::FilterTypeToString(Data->BreakpointInstance.FilterType));
+				}
+				return LOCTEXT("BreakpointInvalid", "Invalid Breakpoint");
+			})
+        ];
+    }
+    else if (ColumnId == "Filter")
+    {
+		if (Data->BreakpointInstance.FilterType == UE::Mass::Debug::FBreakpoint::EFilterType::Query)
+		{
+			return SNew(SSearchableComboBox)
+            .OptionsSource(&DebuggerModel.Pin()->QueryNames)
+            .OnGenerateWidget_Lambda([](TSharedPtr<FString> Str)
+			{
+				if (Str)
+				{
+					return SNew(STextBlock)
+						.Text(FText::FromString(*Str));
+				}
+				return SNew(STextBlock)
+					.Text(LOCTEXT("QueryNone", "None"));
+			})
+            .OnSelectionChanged_Lambda([WeakData = BreakpointData, WeakModel = DebuggerModel](TSharedPtr<FString> Str, ESelectInfo::Type)
+            {
+				if (TSharedPtr<FMassDebuggerBreakpointData> Data = WeakData.Pin())
+				{
+					if (Str)
+					{
+						Data->FilterName = *Str;
+					}
+					else
+					{
+						Data->FilterName.Empty();
+					}
+
+					if (TSharedPtr<FMassDebuggerModel> Model = WeakModel.Pin())
+					{
+						Data->ReconcileDataFromNames(*Model);
+						Model->ApplyBreakpointsToCurrentEnvironment();
+					}
+				}
+            })
             [
-                // Display Entity handle or "Selected Entity"
                 SNew(STextBlock)
-				.Text(InItem->bBreakForSelectedEntity
-					? LOCTEXT("SelectedEntity", "Selected Entity")
-					: FText::FromString(InItem->Entity.DebugGetDescription()))
-				.ToolTip(FSlateApplicationBase::Get().MakeToolTip(InItem->bBreakForSelectedEntity
-					? LOCTEXT("SelectedEntityGameplayDebugger", "Selected Entity (Gameplay Debugger)")
-					: FText::FromString(InItem->Entity.DebugGetDescription())))
-            ]
-#else
-			SNew(STextBlock)
-				.Text(LOCTEXT("MassEntityDebuggingNotEnabled", "Mass Entity Debugging Not Enabled for this configuration"))
-#endif
-		];
+				.Text_Lambda([WeakData = BreakpointData]()
+                {
+					if (TSharedPtr<FMassDebuggerBreakpointData> Data = WeakData.Pin())
+					{
+						if (!Data->FilterName.IsEmpty())
+						{
+							return FText::FromString(Data->FilterName);
+						}
+					}
+                    return LOCTEXT("SelectQuery","Select Query");
+                })
+            ];
+		}
+    }
+	else if (ColumnId == "HitCount")
+	{
+		return SAssignNew(HitCount, STextBlock)
+		.Text_Lambda([WeakData = BreakpointData, WeakModel = DebuggerModel]()
+		{
+			if (TSharedPtr<FMassDebuggerBreakpointData> Data = WeakData.Pin())
+			{
+				if (TSharedPtr<FMassDebuggerModel> Model = WeakModel.Pin(); !Model->IsStale())
+				{
+					UE::Mass::Debug::FBreakpoint* EngineBreak = FMassDebugger::FindBreakpoint(*Model->Environment->GetEntityManager(), Data->BreakpointInstance.Handle);
+					if (EngineBreak)
+					{
+						return FText::AsNumber(EngineBreak->HitCount);
+					}
+				}
+			}
+			return FText::AsNumber(0);
+		});
+	}
+
+    return SNullWidget::NullWidget;
 }
 
-#undef LOCTEXT_NAMESPACE
 
+void SMassBreakpointsView::SBreakpointTableRow::Tick(const FGeometry& AllottedGeometry, const double InCurrentTime, const float InDeltaTime)
+{
+	SMultiColumnTableRow<TSharedPtr<FMassDebuggerBreakpointData>>::Tick(AllottedGeometry, InCurrentTime, InDeltaTime);
+
+	if (HitCount.IsValid())
+	{
+		HitCount->Invalidate(EInvalidateWidgetReason::Paint);
+	}
+}
+
+#undef LOCTEXT_NAMESPACE
diff --git a/Engine/Source/Editor/MassEntityDebugger/Private/SMassBreakpointsView.h b/Engine/Source/Editor/MassEntityDebugger/Private/SMassBreakpointsView.h
--- a/Engine/Source/Editor/MassEntityDebugger/Private/SMassBreakpointsView.h
+++ b/Engine/Source/Editor/MassEntityDebugger/Private/SMassBreakpointsView.h
@@ -3,14 +3,16 @@
 
 #include "SMassDebuggerViewBase.h"
 #include "Widgets/Views/SListView.h"
-#include "MassEntityHandle.h"
+#include "Widgets/Input/SComboBox.h"
+#include "Widgets/Input/SCheckBox.h"
+#include "SSearchableComboBox.h"
+#include "MassDebugger.h"
 
+class SSearchableComboBox;
 struct FMassDebuggerModel;
-struct FMassDebuggerArchetypeData;
 struct FMassDebuggerProcessorData;
-class ITableRow;
-class STableViewBase;
-class UMassProcessor;
+struct FMassDebuggerFragmentData;
+struct FMassDebuggerBreakpointData;
 
 class SMassBreakpointsView : public SMassDebuggerViewBase
 {
@@ -20,38 +22,59 @@ public:
 	}
 	SLATE_END_ARGS()
 
-	void Construct(const FArguments& InArgs, TSharedRef<FMassDebuggerModel> InDebuggerModel);
+	void Construct(const FArguments& InArgs, TSharedRef<FMassDebuggerModel> InModel);
 
 protected:
+	// SMassDebuggerViewBase overrides
+	virtual void OnRefresh() override
+	{
+		RefreshBreakpoints();
+	}
 
-	/** Unused pure virtual function */
-	virtual void OnRefresh() override;
-
-	/** Unused pure virtual function */
-	virtual void OnProcessorsSelected(TConstArrayView<TSharedPtr<FMassDebuggerProcessorData>> SelectedProcessors, ESelectInfo::Type SelectInfo) override
+	virtual void OnProcessorsSelected(TConstArrayView<TSharedPtr<FMassDebuggerProcessorData>> SelectedProcessors, ESelectInfo::Type SelectInfo)
 	{
 	}
 
-	/** Unused pure virtual function */
-	virtual void OnArchetypesSelected(TConstArrayView<TSharedPtr<FMassDebuggerArchetypeData>> SelectedArchetypes, ESelectInfo::Type SelectInfo) override
+	virtual void OnArchetypesSelected(TConstArrayView<TSharedPtr<FMassDebuggerArchetypeData>> SelectedArchetypes, ESelectInfo::Type SelectInfo)
 	{
 	}
 
 private:
-	FReply HandleAddWriteBreakpointClicked();
-	FReply ClearBreakpointsClicked();
-
-	struct FBreakpointDisplay
+	friend class SBreakpointTableRow;
+	class SBreakpointTableRow : public SMultiColumnTableRow<TSharedPtr<FMassDebuggerBreakpointData>>
 	{
-		TWeakObjectPtr<const UMassProcessor> Processor;
-		TWeakObjectPtr<const UScriptStruct> WriteFragment;
-		/** Specifies that this breakpoint is not for a specific entity, but will trigger for any entity selected in the mass or gameplay debugger*/
-		bool bBreakForSelectedEntity = false;
-		FMassEntityHandle Entity;
+	public:
+		SLATE_BEGIN_ARGS(SBreakpointTableRow)
+		{
+		}
+		SLATE_END_ARGS()
+
+		void Construct(const FArguments& InArgs, const TSharedRef<STableViewBase>& InOwnerTableView, TSharedRef<SMassBreakpointsView> InParentView, TSharedRef<FMassDebuggerBreakpointData> InBreakpointData, TSharedRef<FMassDebuggerModel> InModel);
+		virtual TSharedRef<SWidget> GenerateWidgetForColumn(const FName& ColumnName) override;
+		virtual void Tick(const FGeometry& AllottedGeometry, const double InCurrentTime, const float InDeltaTime) override;
+
+	private:
+		TSharedPtr<STextBlock> HitCount;
+
+		TWeakPtr<FMassDebuggerBreakpointData> BreakpointData;
+		TWeakPtr<FMassDebuggerModel> DebuggerModel;
+		TWeakPtr<SMassBreakpointsView> ParentView;
 	};
-	TArray<TSharedPtr<FBreakpointDisplay>> WriteBreakpoints;
-	TSharedRef<ITableRow> OnGenerateBreakpointRow(TSharedPtr<FBreakpointDisplay> InItem, const TSharedRef<STableViewBase>& OwnerTable);
-	TSharedPtr<SListView<TSharedPtr<FBreakpointDisplay>>> WriteBreakpointsListView;
 
 	void RefreshBreakpoints();
+	FReply ClearBreakpointsClicked();
+	FReply SaveBreakpointsClicked();
+	FReply PasteBreakpointClicked();
+	FReply NewBreakpointClicked();
+
+	TSharedRef<ITableRow> OnGenerateBreakpointRow(
+		TSharedPtr<FMassDebuggerBreakpointData> InItem,
+		const TSharedRef<STableViewBase>& OwnerTable);
+
+	TSharedPtr<SListView<TSharedPtr<FMassDebuggerBreakpointData>>> BreakpointsListView;
+
+	TArray<TSharedPtr<FString>> TriggerTypeOptions;
+	TMap<UE::Mass::Debug::FBreakpoint::ETriggerType, TSharedPtr<FString>> TriggerMap;
+	TArray<TSharedPtr<FString>> FilterTypeOptions;
+	TMap<UE::Mass::Debug::FBreakpoint::EFilterType, TSharedPtr<FString>> FilterMap;
 };
diff --git a/Engine/Source/Editor/MassEntityDebugger/Private/SMassDebugger.cpp b/Engine/Source/Editor/MassEntityDebugger/Private/SMassDebugger.cpp
--- a/Engine/Source/Editor/MassEntityDebugger/Private/SMassDebugger.cpp
+++ b/Engine/Source/Editor/MassEntityDebugger/Private/SMassDebugger.cpp
@@ -52,6 +52,13 @@ namespace UE::Mass::Debugger::Private
 	}
 }
 
+static TAutoConsoleVariable<bool> CVarMassDebuggerAutoSwitchToNewEnvironments(
+	TEXT("MassDebugger.AutoSwitchToNewEnvironments"),
+	false,
+	TEXT("If enabled, the Mass Debugger will automatically switch to newly registered environments."),
+	ECVF_Default
+);
+
 //----------------------------------------------------------------------//
 // SMassDebugger
 //----------------------------------------------------------------------//
@@ -70,6 +77,11 @@ SMassDebugger::~SMassDebugger()
 #endif // WITH_MASSENTITY_DEBUG
 }
 
+bool SMassDebugger::IsAutoSwitchEnvironmentsEnabled()
+{
+	return CVarMassDebuggerAutoSwitchToNewEnvironments.GetValueOnAnyThread();
+}
+
 void SMassDebugger::Construct(const FArguments& InArgs, const TSharedRef<SDockTab>& ConstructUnderMajorTab, const TSharedPtr<SWindow>& ConstructUnderWindow)
 {
 	using namespace UE::Mass::Debugger::Private;
@@ -226,6 +238,27 @@ TSharedRef<SDockTab> SMassDebugger::SpawnToolbar(const FSpawnTabArgs& Args)
 		[
 			ToolBarBuilder.MakeWidget()
 		]
+		+ SHorizontalBox::Slot()
+		.HAlign(HAlign_Right)
+		.VAlign(VAlign_Center)
+		.Padding(2.f)
+		.AutoWidth()
+		[
+			SNew(SCheckBox)
+			.IsChecked_Lambda([this]()
+			{
+				return CVarMassDebuggerAutoSwitchToNewEnvironments.GetValueOnAnyThread()
+					? ECheckBoxState::Checked : ECheckBoxState::Unchecked;
+			})
+			.OnCheckStateChanged_Lambda([this](ECheckBoxState NewState)
+			{
+				CVarMassDebuggerAutoSwitchToNewEnvironments->Set(NewState == ECheckBoxState::Checked, ECVF_SetByConsole);
+			})
+			.ToolTipText(LOCTEXT("AutoSwitchEnvironments_Tooltip", "Automatically switch to new environments as they are registered"))
+			[
+				SNew(STextBlock).Text(LOCTEXT("AutoSwitchEnvironments_Label", "Auto-Switch"))
+			]
+		]
 		+SHorizontalBox::Slot()
 		.HAlign(HAlign_Right)
 		.VAlign(VAlign_Center)
@@ -283,10 +316,15 @@ void SMassDebugger::OnEntityManagerInitialized(const FMassEntityManager& EntityM
 	const UWorld* World = EntityManager.GetWorld();
 	if (World && UE::Mass::Debugger::Private::IsSupportedWorldType(World->WorldType))
 	{
-		EnvironmentsList.Add(MakeShareable(new FMassDebuggerEnvironment(EntityManager.AsShared())));
+		const int32 Index = EnvironmentsList.Add(MakeShareable(new FMassDebuggerEnvironment(EntityManager.AsShared())));
 		if (EnvironmentComboBox.IsValid())
 		{
 			EnvironmentComboBox->RefreshOptions();
+			if (IsAutoSwitchEnvironmentsEnabled())
+			{
+				EnvironmentComboBox->SetSelectedItem(EnvironmentsList[Index]);
+				HandleEnvironmentChanged(EnvironmentsList[Index], ESelectInfo::Type::Direct);
+			}
 		}
 	}
 }
diff --git a/Engine/Source/Editor/MassEntityDebugger/Private/SMassDebugger.h b/Engine/Source/Editor/MassEntityDebugger/Private/SMassDebugger.h
--- a/Engine/Source/Editor/MassEntityDebugger/Private/SMassDebugger.h
+++ b/Engine/Source/Editor/MassEntityDebugger/Private/SMassDebugger.h
@@ -66,6 +66,8 @@ public:
 	void ShowSelectedView() const;
 	void ShowQueryEditorView() const;
 
+	static bool IsAutoSwitchEnvironmentsEnabled();
+
 protected:
 	TSharedRef<SDockTab> SpawnToolbar(const FSpawnTabArgs& Args);
 	TSharedRef<SDockTab> SpawnProcessorsTab(const FSpawnTabArgs& Args);
diff --git a/Engine/Source/Editor/MassEntityDebugger/Private/SMassFragmentsView.cpp b/Engine/Source/Editor/MassEntityDebugger/Private/SMassFragmentsView.cpp
--- a/Engine/Source/Editor/MassEntityDebugger/Private/SMassFragmentsView.cpp
+++ b/Engine/Source/Editor/MassEntityDebugger/Private/SMassFragmentsView.cpp
@@ -25,7 +25,6 @@ namespace UE::MassDebugger::FragmentsView
 	namespace Private
 	{
 		static const FLazyName ColumnName = TEXT("Name");
-		static const FLazyName ColumnDebugControls = TEXT("DebugControls");
 		static const FLazyName ColumnNumArchetypes = TEXT("Archetypes");
 		static const FLazyName ColumnNumEntities = TEXT("Entities");
 	}
@@ -41,13 +40,6 @@ void SFragmentsView::Construct(const FArguments& InArgs, TSharedRef<FMassDebugge
 		.FillWidth(0.5f)
 		.SortMode(this, &SFragmentsView::GetColumnSortMode, Private::ColumnName.Resolve())
 		.OnSort(this, &SFragmentsView::OnColumnSortModeChanged)
-		+ SHeaderRow::Column(Private::ColumnDebugControls)
-		.DefaultLabel(LOCTEXT("FragmentColumnDebugControls", "Debug Controls"))
-		.FixedWidth(30.f)
-		.HAlignHeader(HAlign_Center)
-		.HAlignCell(HAlign_Center)
-		.VAlignCell(VAlign_Center)
-		.SortMode(EColumnSortMode::None)
 		+ SHeaderRow::Column(Private::ColumnNumArchetypes)
 		.DefaultLabel(LOCTEXT("FragmentColumnArchetypes", "Archetypes"))
 		.FillWidth(0.25f)
@@ -98,11 +90,6 @@ EColumnSortMode::Type SFragmentsView::GetColumnSortMode(FName ColumnId) const
 
 void SFragmentsView::OnColumnSortModeChanged(const EColumnSortPriority::Type SortPriority, const FName& ColumnId, const EColumnSortMode::Type InSortMode)
 {
-	if (ColumnId == Private::ColumnDebugControls)
-	{
-		return;
-	}
-
 	if (SortPriority == EColumnSortPriority::Primary)
 	{
 		SortByColumnId = ColumnId;
@@ -218,41 +205,6 @@ TSharedRef<SWidget> SFragmentsTableRow::GenerateWidgetForColumn(const FName& InC
 			.Text(FText::AsNumber(FragmentDataPtr->NumEntities))
 			.Justification(ETextJustify::Right);
 	}
-	else if (InColumnName == Private::ColumnDebugControls)
-	{
-		SAssignNew(CellWidget, SButton)
-			.ToolTipText(LOCTEXT("BreakpointTooltip", "Break on writes for selected entity."))
-			.ButtonStyle(FAppStyle::Get(), "NoBorder")
-			.ContentPadding(0)
-			.OnClicked_Lambda([WeakEntityManager = FragmentDataPtr->EntityManager, Fragment = FragmentDataPtr->Fragment]()
-			{
-				TSharedPtr<const FMassEntityManager> StrongEntityManager = WeakEntityManager.Pin();
-				TStrongObjectPtr<const UScriptStruct> StrongFragment = Fragment.Pin();
-				if (StrongEntityManager && StrongFragment)
-				{
-					FMassDebugger::ToggleFragmentWriteBreakForSelectedEntity(*StrongEntityManager, StrongFragment.Get());
-				}
-				return FReply::Handled();
-			})
-			[
-				SNew(SImage)
-				.ColorAndOpacity_Lambda([WeakEntityManager = FragmentDataPtr->EntityManager, Fragment = FragmentDataPtr->Fragment]()
-				{
-					TSharedPtr<const FMassEntityManager> StrongEntityManager = WeakEntityManager.Pin();
-					TStrongObjectPtr<const UScriptStruct> StrongFragment = Fragment.Pin();
-					if (StrongEntityManager && StrongFragment)
-					{
-						if(FMassDebugger::IsFragmentWriteBreakForSelectedEntitySet(*StrongEntityManager, StrongFragment.Get()))
-						{
-							return FSlateColor(FLinearColor(1.f, 0.f, 0.f, 1.f));
-						}
-					}
-					return FSlateColor(FLinearColor(1.f, 1.f, 1.f, 0.8f));
-				})
-				.Image(FAppStyle::GetBrush("Icons.FilledCircle"))
-			];
-	}
-
 
 #endif //WITH_MASSENTITY_DEBUG
 
diff --git a/Engine/Source/Editor/MassEntityDebugger/Private/SMassQueryEditor.cpp b/Engine/Source/Editor/MassEntityDebugger/Private/SMassQueryEditor.cpp
--- a/Engine/Source/Editor/MassEntityDebugger/Private/SMassQueryEditor.cpp
+++ b/Engine/Source/Editor/MassEntityDebugger/Private/SMassQueryEditor.cpp
@@ -196,6 +196,10 @@ namespace UE::MassDebugger
 		{
 			Query->Name = InText.ToString();
 			NotifyRequirementsChanged();
+			if (TSharedPtr<FMassDebuggerModel> Model = DebuggerModel.Pin())
+			{
+				Model->RefreshQueries();
+			}
 		}
 	}
 
diff --git a/Engine/Source/Editor/MassEntityDebugger/Private/SMassQueryEditor.h b/Engine/Source/Editor/MassEntityDebugger/Private/SMassQueryEditor.h
--- a/Engine/Source/Editor/MassEntityDebugger/Private/SMassQueryEditor.h
+++ b/Engine/Source/Editor/MassEntityDebugger/Private/SMassQueryEditor.h
@@ -8,7 +8,6 @@
 #include "MassDebuggerModel.h"
 #include "Widgets/Views/SListView.h"
 #include "Widgets/Views/STableRow.h"
-#include "Widgets/Input/SEditableTextBox.h"
 
 class SSearchableComboBox;
 
diff --git a/Engine/Source/Editor/MassEntityDebugger/Private/SMassQueryEditorView.cpp b/Engine/Source/Editor/MassEntityDebugger/Private/SMassQueryEditorView.cpp
--- a/Engine/Source/Editor/MassEntityDebugger/Private/SMassQueryEditorView.cpp
+++ b/Engine/Source/Editor/MassEntityDebugger/Private/SMassQueryEditorView.cpp
@@ -22,6 +22,7 @@ namespace UE::MassDebugger
 	{
 		DebuggerModel = InDebuggerModel;
 		InDebuggerModel->OnRefreshDelegate.AddSP(this, &SQueryEditorView::OnRefresh);
+		InDebuggerModel->OnQueriesChangedDelegate.AddSP(this, &SQueryEditorView::OnRefresh);
 
 		ChildSlot
 		[
@@ -40,7 +41,7 @@ namespace UE::MassDebugger
 					[
 						SNew(SButton)
 						.OnClicked(this, &SQueryEditorView::OnSaveQuery)
-						.ToolTipText(LOCTEXT("SaveQuery_Tooltip", "Save selected query to disk"))
+						.ToolTipText(LOCTEXT("SaveQuery_Tooltip", "Save selected query"))
 						.Content()
 						[
 							SNew(SImage)
@@ -51,7 +52,7 @@ namespace UE::MassDebugger
 					[
 						SNew(SButton)
 						.OnClicked(this, &SQueryEditorView::OnSaveAllQueries)
-						.ToolTipText(LOCTEXT("SaveAllQueries_Tooltip", "Save all queries to disk"))
+						.ToolTipText(LOCTEXT("SaveAllQueries_Tooltip", "Save all queries"))
 						.Content()
 						[
 							SNew(SImage)
@@ -66,7 +67,7 @@ namespace UE::MassDebugger
 						.Content()
 						[
 							SNew(SImage)
-							.Image(FAppStyle::Get().GetBrush("Icons.Clipboard"))
+							.Image(FAppStyle::Get().GetBrush("GenericCommands.Copy"))
 						]
 					]
 					+ SHorizontalBox::Slot().AutoWidth().Padding(2)
@@ -100,28 +101,31 @@ namespace UE::MassDebugger
 				.Padding(2)
 				[
 					SAssignNew(QueryListView, SListView<TSharedPtr<FEditableQuery>>)
-					.ListItemsSource(&QueryList)
+					.ListItemsSource(&InDebuggerModel->QueryList)
 					.SelectionMode(ESelectionMode::Single)
 					.OnGenerateRow(this, &SQueryEditorView::OnGenerateQueryRow)
-					.OnSelectionChanged_Lambda([this](TSharedPtr<FEditableQuery> InQuery, ESelectInfo::Type InSelectType)
+					.OnSelectionChanged_Lambda([this, WeakModel = InDebuggerModel.ToWeakPtr()](TSharedPtr<FEditableQuery> InQuery, ESelectInfo::Type InSelectType)
 					{
-						if (!InQuery.IsValid())
+						if(TSharedPtr<FMassDebuggerModel> Model = WeakModel.Pin())
 						{
-							if (SelectedQuery.IsValid() && QueryList.Contains(SelectedQuery))
+							if (!InQuery.IsValid())
 							{
-								QueryListView->SetSelection(SelectedQuery);
+								if (SelectedQuery.IsValid() && Model->QueryList.Contains(SelectedQuery))
+								{
+									QueryListView->SetSelection(SelectedQuery);
+								}
+								else
+								{
+									SelectedQuery = Model->QueryList.Num() > 0 ? Model->QueryList[0] : InQuery;
+									QueryEditor->SetQuery(SelectedQuery);
+								}
 							}
 							else
 							{
-								SelectedQuery = QueryList.Num() > 0 ? QueryList[0] : InQuery;
+								SelectedQuery = InQuery;
 								QueryEditor->SetQuery(SelectedQuery);
 							}
 						}
-						else
-						{
-							SelectedQuery = InQuery;
-							QueryEditor->SetQuery(SelectedQuery);
-						}
 					})
 				]
 			]
@@ -136,7 +140,6 @@ namespace UE::MassDebugger
 			]
 		];
 
-		LoadQueriesFromDisk();
 		InDebuggerModel->RegisterQueryEditor(SharedThis(this));
 	}
 
@@ -194,28 +197,36 @@ namespace UE::MassDebugger
 		TSharedPtr<FEditableQuery> NewQuery = MakeShared<FEditableQuery>();
 		NewQuery->Name = TEXT("NewQuery");
 
-		QueryList.Add(NewQuery);
-		if (QueryListView.IsValid())
+		if (TSharedPtr<FMassDebuggerModel> Model = DebuggerModel.Pin())
 		{
-			QueryListView->RequestListRefresh();
-			QueryListView->SetSelection(NewQuery);
-			QueryListView->RequestScrollIntoView(NewQuery);
+			Model->QueryList.Add(NewQuery);
+			if (QueryListView.IsValid())
+			{
+				QueryListView->RequestListRefresh();
+				QueryListView->SetSelection(NewQuery);
+				QueryListView->RequestScrollIntoView(NewQuery);
+			}
+			Model->RefreshQueries();
 		}
 		return FReply::Handled();
 	}
 
 	FReply SQueryEditorView::OnDeleteQuery(TSharedPtr<FEditableQuery> InItem)
 	{
-		QueryList.Remove(InItem);
-		if (QueryListView.IsValid())
+		if (TSharedPtr<FMassDebuggerModel> Model = DebuggerModel.Pin())
 		{
-			QueryListView->RequestListRefresh();
-		}
+			Model->QueryList.Remove(InItem);
+			if (QueryListView.IsValid())
+			{
+				QueryListView->RequestListRefresh();
+			}
 
-		if (SelectedQuery == InItem && QueryEditor.IsValid())
-		{
-			SelectedQuery.Reset();
-			QueryEditor->SetQuery(SelectedQuery);
+			if (SelectedQuery == InItem && QueryEditor.IsValid())
+			{
+				SelectedQuery.Reset();
+				QueryEditor->SetQuery(SelectedQuery);
+			}
+			Model->RefreshQueries();
 		}
 		return FReply::Handled();
 	}
@@ -248,243 +259,21 @@ namespace UE::MassDebugger
 	{
 	}
 
-	void FEditableQuery::InitFromEntityQuery(const FMassEntityQuery& InQuery, FMassDebuggerModel& DebuggerModel)
-	{
-		FragmentRequirements.Reset();
-
-		for (const FMassFragmentRequirementDescription& Desc : InQuery.GetFragmentRequirements())
-		{
-			TSharedPtr<FFragmentEntry> Entry = MakeShared<FFragmentEntry>();
-			Entry->StructType = Desc.StructType;
-			Entry->AccessMode = Desc.AccessMode;
-			Entry->Presence = Desc.Presence;
-			FragmentRequirements.Add(Entry);
-		}
-
-		const FMassTagBitSet& AllTags = InQuery.GetRequiredAllTags();
-		const FMassTagBitSet& AnyTags = InQuery.GetRequiredAnyTags();
-		const FMassTagBitSet& OptionalTags = InQuery.GetRequiredOptionalTags();
-		const FMassTagBitSet& NoneTags = InQuery.GetRequiredNoneTags();
-
-		for (const TSharedPtr<FMassDebuggerFragmentData>& TagData : DebuggerModel.CachedTagData)
-		{
-			const UScriptStruct& StructType = *TagData->Fragment.Get();
-			EMassFragmentPresence Presence = EMassFragmentPresence::All;
-
-			if (AllTags.Contains(StructType))
-			{
-				Presence = EMassFragmentPresence::All;
-			}
-			else if (AnyTags.Contains(StructType))
-			{
-				Presence = EMassFragmentPresence::Any;
-			}
-			else if (OptionalTags.Contains(StructType))
-			{
-				Presence = EMassFragmentPresence::Optional;
-			}
-			else if (NoneTags.Contains(StructType))
-			{
-				Presence = EMassFragmentPresence::None;
-			}
-			else
-			{
-				continue;
-			}
-
-			TSharedPtr<FFragmentEntry> Entry = MakeShared<FFragmentEntry>();
-			Entry->StructType = &StructType;
-			Entry->Presence = Presence;
-			TagRequirements.Add(Entry);
-		}
-	}
-
-	FMassEntityQuery FEditableQuery::BuildEntityQuery(const TSharedRef<FMassEntityManager> EntityManager)
-	{
-		FMassEntityQuery Query;
-
-		Query.Initialize(EntityManager);
-
-		for (TSharedPtr<FFragmentEntry>& Fragment : FragmentRequirements)
-		{
-			if (ensure(Fragment->StructType.IsValid()))
-			{
-				Query.AddRequirement(Fragment->StructType.Get(), Fragment->AccessMode, Fragment->Presence);
-			}
-		}
-
-		for (TSharedPtr<FFragmentEntry>& TagEntry : TagRequirements)
-		{
-			if (ensure(TagEntry->StructType.IsValid()))
-			{
-				Query.AddTagRequirement(*TagEntry->StructType, TagEntry->Presence);
-			}
-		}
-		return Query;
-	}
-
-	TSharedPtr<FJsonObject> FFragmentEntry::SerializeToJson() const
-	{
-		if (!StructType.IsValid())
-		{
-			return nullptr;
-		}
-		TSharedPtr<FJsonObject> JsonObject = MakeShared<FJsonObject>();
-
-		JsonObject->SetStringField(TEXT("StructType"), StructType->GetPathName());
-		JsonObject->SetStringField(TEXT("AccessMode"), StaticEnum<EMassFragmentAccess>()->GetNameStringByValue((int64)AccessMode));
-		JsonObject->SetStringField(TEXT("Presence"), StaticEnum<EMassFragmentPresence>()->GetNameStringByValue((int64)Presence));
-
-		return JsonObject;
-	}
-
-	bool FFragmentEntry::DeserializeFromJson(const TSharedPtr<FJsonObject>& JsonObject)
-	{
-		if (!JsonObject.IsValid())
-		{
-			return false;
-		}
-		const FString StructPath = JsonObject->GetStringField(TEXT("StructType"));
-		if (!StructPath.IsEmpty())
-		{
-			UScriptStruct* Found = FindObject<UScriptStruct>(nullptr, *StructPath);
-			StructType = Found;
-		}
-		else
-		{
-			StructType = nullptr;
-		}
-
-		const FString AccessStr = JsonObject->GetStringField(TEXT("AccessMode"));
-		UEnum* AccessEnum = StaticEnum<EMassFragmentAccess>();
-		int64 AccessVal = AccessEnum->GetValueByNameString(AccessStr);
-		AccessMode = (EMassFragmentAccess)AccessVal;
-
-		const FString PresenceStr = JsonObject->GetStringField(TEXT("Presence"));
-		UEnum* PresenceEnum = StaticEnum<EMassFragmentPresence>();
-		int64 PresenceVal = PresenceEnum->GetValueByNameString(PresenceStr);
-		Presence = (EMassFragmentPresence)PresenceVal;
-
-		return true;
-	}
-
-	bool FFragmentEntry::DeserializeFromJsonString(const FString& JsonString)
-	{
-		TSharedPtr<FJsonObject> JsonObject;
-		TSharedRef<TJsonReader<>> Reader = TJsonReaderFactory<>::Create(JsonString);
-
-		if (!FJsonSerializer::Deserialize(Reader, JsonObject) || !JsonObject.IsValid())
-		{
-			return false;
-		}
-
-		return DeserializeFromJson(JsonObject);
-	}
-
-	TSharedPtr<FJsonObject> FEditableQuery::SerializeToJson() const
-	{
-		TSharedPtr<FJsonObject> JsonObject = MakeShared<FJsonObject>();
-		JsonObject->SetStringField(TEXT("Name"), Name);
-
-		TArray<TSharedPtr<FJsonValue>> JsonFragments;
-		JsonFragments.Reserve(FragmentRequirements.Num());
-		for (auto& Entry : FragmentRequirements)
-		{
-			JsonFragments.Add(MakeShared<FJsonValueObject>(Entry->SerializeToJson()));
-		}
-		JsonObject->SetArrayField(TEXT("FragmentRequirements"), MoveTemp(JsonFragments));
-
-		TArray<TSharedPtr<FJsonValue>> JsonTags;
-		JsonTags.Reserve(TagRequirements.Num());
-		for (auto& Entry : TagRequirements)
-		{
-			JsonTags.Add(MakeShared<FJsonValueObject>(Entry->SerializeToJson()));
-		}
-		JsonObject->SetArrayField(TEXT("TagRequirements"), MoveTemp(JsonTags));
-
-		return JsonObject;
-	}
-
-	bool FEditableQuery::DeserializeFromJson(const TSharedPtr<FJsonObject>& JsonObject)
-	{
-		if (!JsonObject.IsValid())
-		{
-			return false;
-		}
-
-		if (JsonObject->HasField(TEXT("Name")))
-		{
-			Name = JsonObject->GetStringField(TEXT("Name"));
-		}
-
-		FragmentRequirements.Reset();
-		const TArray<TSharedPtr<FJsonValue>>* JsonFragments = nullptr;
-		if (JsonObject->TryGetArrayField(TEXT("FragmentRequirements"), JsonFragments))
-		{
-			for (const TSharedPtr<FJsonValue>& Val : *JsonFragments)
-			{
-				if (TSharedPtr<FJsonObject> Obj = Val->AsObject())
-				{
-					TSharedPtr<FFragmentEntry> Entry = MakeShared<FFragmentEntry>();
-					if (!Entry->DeserializeFromJson(Obj))
-					{
-						return false;
-					}
-					Entry->bShowAccessMode = true;
-					FragmentRequirements.Add(MoveTemp(Entry));
-				}
-			}
-		}
-
-		TagRequirements.Reset();
-		const TArray<TSharedPtr<FJsonValue>>* JsonTags = nullptr;
-		if (JsonObject->TryGetArrayField(TEXT("TagRequirements"), JsonTags))
-		{
-			for (const TSharedPtr<FJsonValue>& Val : *JsonTags)
-			{
-				if (TSharedPtr<FJsonObject> Obj = Val->AsObject())
-				{
-					TSharedPtr<FFragmentEntry> Entry = MakeShared<FFragmentEntry>();
-					if (!Entry->DeserializeFromJson(Obj))
-					{
-						return false;
-					}
-					Entry->bShowAccessMode = false;
-					TagRequirements.Add(MoveTemp(Entry));
-				}
-			}
-		}
-
-		return true;
-	}
-
-	bool FEditableQuery::DeserializeFromJsonString(const FString& JsonString)
-	{
-		TSharedPtr<FJsonObject> JsonObject;
-		TSharedRef<TJsonReader<>> Reader = TJsonReaderFactory<>::Create(JsonString);
-
-		if (FJsonSerializer::Deserialize(Reader, JsonObject) && JsonObject.IsValid())
-		{
-			return DeserializeFromJson(JsonObject);
-		}
-
-		return false;
-	}
-
 	void SQueryEditorView::ShowQuery(const FMassEntityQuery& InQuery, const FString& InQueryName)
 	{
-		if (TSharedPtr<FMassDebuggerModel> SharedModel = DebuggerModel.Pin())
+		if (TSharedPtr<FMassDebuggerModel> Model = DebuggerModel.Pin())
 		{
 			TSharedPtr<FEditableQuery> NewQuery = MakeShared<FEditableQuery>();
 			NewQuery->Name = InQueryName;
-			NewQuery->InitFromEntityQuery(InQuery, *SharedModel);
-			QueryList.Add(NewQuery);
+			NewQuery->InitFromEntityQuery(InQuery, *Model);
+			Model->QueryList.Add(NewQuery);
 			if (QueryListView.IsValid())
 			{
 				QueryListView->RequestListRefresh();
 				QueryListView->SetSelection(NewQuery);
 				QueryListView->RequestScrollIntoView(NewQuery);
 			}
+			Model->RefreshQueries();
 		}
 	}
 
@@ -543,24 +332,26 @@ namespace UE::MassDebugger
 
 	FReply SQueryEditorView::OnSaveAllQueries()
 	{
-		const FString FileName = TEXT("Queries.json");
-		const FString SavedDir = FPaths::ProjectSavedDir() / TEXT("MassDebugger");
-		IFileManager::Get().MakeDirectory(*SavedDir, /*Tree=*/true);
-		const FString FullPath = SavedDir / FileName;
-
-		TSharedPtr<FJsonObject> RootObject = MakeShared<FJsonObject>();
-		TArray<TSharedPtr<FJsonValue>> QueriesArray;
-		for (auto& Q : QueryList)
+		if (TSharedPtr<FMassDebuggerModel> Model = DebuggerModel.Pin())
 		{
-			QueriesArray.Add(MakeShared<FJsonValueObject>(Q->SerializeToJson()));
-		}
-		RootObject->SetArrayField(TEXT("Queries"), MoveTemp(QueriesArray));
+			const FString FileName = TEXT("Queries.json");
+			const FString SavedDir = FPaths::ProjectSavedDir() / TEXT("MassDebugger");
+			IFileManager::Get().MakeDirectory(*SavedDir, /*Tree=*/true);
+			const FString FullPath = SavedDir / FileName;
 
-		FString OutString;
-		const TSharedRef<TJsonWriter<>> Writer = TJsonWriterFactory<>::Create(&OutString);
-		FJsonSerializer::Serialize(RootObject.ToSharedRef(), Writer);
-		FFileHelper::SaveStringToFile(OutString, *FullPath);
+			TSharedPtr<FJsonObject> RootObject = MakeShared<FJsonObject>();
+			TArray<TSharedPtr<FJsonValue>> QueriesArray;
+			for (TSharedPtr<FEditableQuery>& Query : Model->QueryList)
+			{
+				QueriesArray.Add(MakeShared<FJsonValueObject>(Query->SerializeToJson()));
+			}
+			RootObject->SetArrayField(TEXT("Queries"), MoveTemp(QueriesArray));
 
+			FString OutString;
+			const TSharedRef<TJsonWriter<>> Writer = TJsonWriterFactory<>::Create(&OutString);
+			FJsonSerializer::Serialize(RootObject.ToSharedRef(), Writer);
+			FFileHelper::SaveStringToFile(OutString, *FullPath);
+		}
 		return FReply::Handled();
 	}
 
@@ -581,56 +372,32 @@ namespace UE::MassDebugger
 
 	FReply SQueryEditorView::OnPasteQuery()
 	{
+		TSharedPtr<FMassDebuggerModel> Model = DebuggerModel.Pin();
+
 		FString ClipboardText;
 		FPlatformApplicationMisc::ClipboardPaste(ClipboardText);
 
 		TSharedPtr<FEditableQuery> NewQuery = MakeShared<FEditableQuery>();
-		if (NewQuery->DeserializeFromJsonString(ClipboardText))
+		if (Model && NewQuery->DeserializeFromJsonString(ClipboardText))
 		{
-			QueryList.Add(NewQuery);
+			Model->QueryList.Add(NewQuery);
 			if (QueryListView.IsValid())
 			{
 				QueryListView->RequestListRefresh();
 				QueryListView->SetSelection(NewQuery);
 				QueryListView->RequestScrollIntoView(NewQuery);
 			}
+			Model->RefreshQueries();
 		}
 		return FReply::Handled();
 	}
 
-	void SQueryEditorView::LoadQueriesFromDisk()
+	void SQueryEditorView::OnQueriesChanged()
 	{
-		const FString SavedDir = FPaths::ProjectSavedDir() / TEXT("MassDebugger");
-		const FString FullPath = SavedDir / TEXT("Queries.json");
-
-		QueryList.Empty();
-
-		if (FPaths::FileExists(FullPath))
+		TSharedPtr<FMassDebuggerModel> Model = DebuggerModel.Pin();
+		if (!Model)
 		{
-			FString FileContent;
-			if (FFileHelper::LoadFileToString(FileContent, *FullPath))
-			{
-				TSharedPtr<FJsonObject> RootObj;
-				TSharedRef<TJsonReader<>> Reader = TJsonReaderFactory<>::Create(FileContent);
-				if (FJsonSerializer::Deserialize(Reader, RootObj) && RootObj.IsValid())
-				{
-					const TArray<TSharedPtr<FJsonValue>>* JsonQueries = nullptr;
-					if (RootObj->TryGetArrayField(TEXT("Queries"), JsonQueries))
-					{
-						for (const auto& Val : *JsonQueries)
-						{
-							if (TSharedPtr<FJsonObject> Obj = Val->AsObject())
-							{
-								TSharedPtr<FEditableQuery> NewQ = MakeShared<FEditableQuery>();
-								if (NewQ->DeserializeFromJson(Obj))
-								{
-									QueryList.Add(MoveTemp(NewQ));
-								}
-							}
-						}
-					}
-				}
-			}
+			return;
 		}
 
 		if (QueryListView.IsValid())
@@ -638,9 +405,9 @@ namespace UE::MassDebugger
 			QueryListView->RequestListRefresh();
 		}
 
-		if (QueryList.Num() > 0)
+		if (Model->QueryList.Num() > 0)
 		{
-			SelectedQuery = QueryList[0];
+			SelectedQuery = Model->QueryList[0];
 			QueryListView->SetSelection(SelectedQuery);
 		}
 		else
diff --git a/Engine/Source/Editor/MassEntityDebugger/Private/SMassQueryEditorView.h b/Engine/Source/Editor/MassEntityDebugger/Private/SMassQueryEditorView.h
--- a/Engine/Source/Editor/MassEntityDebugger/Private/SMassQueryEditorView.h
+++ b/Engine/Source/Editor/MassEntityDebugger/Private/SMassQueryEditorView.h
@@ -11,29 +11,6 @@ class FJsonObject;
 
 namespace UE::MassDebugger
 {
-	struct FFragmentEntry
-	{
-		TSharedPtr<FJsonObject> SerializeToJson() const;
-		bool DeserializeFromJson(const TSharedPtr<FJsonObject>& JsonObject);
-		bool DeserializeFromJsonString(const FString& JsonString);
-		TWeakObjectPtr<const UScriptStruct> StructType = nullptr;
-		EMassFragmentAccess AccessMode = EMassFragmentAccess::ReadOnly;
-		EMassFragmentPresence Presence = EMassFragmentPresence::All;
-		bool bShowAccessMode = true;
-	};
-
-	struct FEditableQuery
-	{
-		TSharedPtr<FJsonObject> SerializeToJson() const;	
-		bool DeserializeFromJson(const TSharedPtr<FJsonObject>& JsonObject);
-		bool DeserializeFromJsonString(const FString& JsonString);
-		void InitFromEntityQuery(const FMassEntityQuery& InQuery, FMassDebuggerModel& DebuggerModel);
-		FString Name;
-		TArray<TSharedPtr<FFragmentEntry>> FragmentRequirements;
-		TArray<TSharedPtr<FFragmentEntry>> TagRequirements;
-		FMassEntityQuery BuildEntityQuery(const TSharedRef<FMassEntityManager> EntityManager);
-	};
-
 	class SQueryEditorView : public SMassDebuggerViewBase
 	{
 	public:
@@ -55,11 +32,9 @@ namespace UE::MassDebugger
 		FReply OnSaveAllQueries();
 		FReply OnCopyQuery();
 		FReply OnPasteQuery();
-		void LoadQueriesFromDisk();
 
 	protected:
 
-		TArray<TSharedPtr<FEditableQuery>> QueryList;
 		TSharedPtr<SListView<TSharedPtr<FEditableQuery>>> QueryListView;
 		TSharedPtr<SQueryEditor> QueryEditor;
 		TSharedPtr<FEditableQuery> SelectedQuery;
@@ -72,6 +47,7 @@ namespace UE::MassDebugger
 		FReply OnEditQuery(TSharedPtr<FEditableQuery> InItem);
 		FReply OnViewEntities(TSharedPtr<FEditableQuery> InItem);
 		void OnRequirementsChanged(TSharedPtr<FEditableQuery>& Query);
+		void OnQueriesChanged();
 	};
 
 } // namespace UE::MassDebugger
diff --git a/Engine/Source/Runtime/MassEntity/Private/MassDebugger.cpp b/Engine/Source/Runtime/MassEntity/Private/MassDebugger.cpp
--- a/Engine/Source/Runtime/MassEntity/Private/MassDebugger.cpp
+++ b/Engine/Source/Runtime/MassEntity/Private/MassDebugger.cpp
@@ -42,6 +42,34 @@ namespace UE::Mass::Debug
 			TEXT("mass.debug.DebugEntity or the gameplay debugger) and storing potential failure results to be viewed in MassDebugger")
 			, ECVF_Cheat }
 	};
+
+	struct FNoDebuggerNotification
+	{
+		static void ConditionallyNotifyUser()
+		{
+			const bool bLocalDebuggerPresent = FPlatformMisc::IsDebuggerPresent();
+			if (bDebuggerPresent != bLocalDebuggerPresent)
+			{
+				bDebuggerPresent = bLocalDebuggerPresent;
+				if (bLocalDebuggerPresent == false)
+				{
+					if (bUserNotified == false)
+					{
+						FMessageDialog::Open(EAppMsgType::Ok, LOCTEXT("NoDebuggerAttached", "Breakpoint set but no debugger is attached."));
+						bUserNotified = true;
+					}
+				}
+				else
+				{
+					bUserNotified = false;
+				}
+			}
+		}
+		static bool bDebuggerPresent;
+		static bool bUserNotified;
+	};
+	bool FNoDebuggerNotification::bDebuggerPresent = false;
+	bool FNoDebuggerNotification::bUserNotified = false;
 	
 
 	FString DebugGetFragmentAccessString(EMassFragmentAccess Access)
@@ -483,7 +511,6 @@ FMassDebugger::FOnEnvironmentEvent FMassDebugger::OnProcessorProviderRegistered;
 FMassDebugger::FOnDebugEvent FMassDebugger::OnDebugEvent;
 TArray<FMassDebugger::FEnvironment> FMassDebugger::ActiveEnvironments;
 UE::FTransactionallySafeMutex FMassDebugger::EntityManagerRegistrationLock;
-bool FMassDebugger::bHasBreakpoint = false;
 TMap<FName, const UScriptStruct*> FMassDebugger::FragmentsByName;
 
 TConstArrayView<FMassEntityQuery*> FMassDebugger::GetProcessorQueries(const UMassProcessor& Processor)
@@ -945,28 +972,49 @@ bool FMassDebugger::DoesArchetypeMatchRequirements(const FMassArchetypeHandle& A
 	return false;
 }
 
-bool FMassDebugger::ShouldProcessorBreak(const FMassEntityManager& EntityManager, const UMassProcessor* Processor, FMassEntityHandle Entity)
+UE::Mass::Debug::FBreakpointHandle FMassDebugger::ShouldProcessorBreak(const FMassEntityManager& EntityManager, const UMassProcessor* Processor, FMassEntityHandle Entity)
 {
-	if (LIKELY(!bHasBreakpoint))
+	if (LIKELY(!UE::Mass::Debug::FBreakpoint::HasBreakpoint()))
 	{
-		return false;
+		return UE::Mass::Debug::FBreakpointHandle::Invalid();
 	}
 
 	FEnvironment& ActiveEnvironment = GetActiveEnvironmentChecked(EntityManager);
 
 	if (LIKELY(!ActiveEnvironment.bHasBreakpoint))
 	{
-		return false;
+		return UE::Mass::Debug::FBreakpointHandle::Invalid();
 	}
 
-	const FMassEntityHandle* Found = ActiveEnvironment.ProcessorBreakpoints.FindPair(Processor, Entity);
+	const bool ProcessorHasBreakpoint = ActiveEnvironment.ProcessorsWithBreakpoints.Contains(Processor);
 	
-	return Found != nullptr;
+	if (ProcessorHasBreakpoint)
+	{
+		UE::Mass::Debug::FBreakpointHandle FoundHandle = 0;
+		for (UE::Mass::Debug::FBreakpoint& BreakPoint : ActiveEnvironment.Breakpoints)
+		{
+			if (BreakPoint.TriggerType == UE::Mass::Debug::FBreakpoint::ETriggerType::ProcessorExecute)
+			{
+				TObjectKey<const UMassProcessor> Proc = BreakPoint.Trigger.Get<TObjectKey<const UMassProcessor>>();
+				if (Proc == Processor && BreakPoint.ApplyEntityFilter(EntityManager, Entity))
+				{
+					++BreakPoint.HitCount;
+					if (BreakPoint.bEnabled)
+					{
+						UE::Mass::Debug::FBreakpoint::LastBreakpointHandle = BreakPoint.Handle;
+						FoundHandle = BreakPoint.Handle;
+					}
+				}
+			}
+		}
+		return FoundHandle;
+	}
+	return UE::Mass::Debug::FBreakpointHandle::Invalid();
 }
 
 bool FMassDebugger::HasAnyProcessorBreakpoints(const FMassEntityManager& EntityManager, const UMassProcessor* Processor)
 {
-	if (LIKELY(!bHasBreakpoint))
+	if (LIKELY(!UE::Mass::Debug::FBreakpoint::HasBreakpoint()))
 	{
 		return false;
 	}
@@ -978,44 +1026,68 @@ bool FMassDebugger::HasAnyProcessorBreakpoints(const FMassEntityManager& EntityM
 		return false;
 	}
 
-	return ActiveEnvironment.ProcessorBreakpoints.Contains(Processor);
+	return ActiveEnvironment.ProcessorsWithBreakpoints.Contains(Processor);
 }
 
-bool FMassDebugger::ShouldBreakOnFragmentWrite(const FMassEntityManager& EntityManager, const UScriptStruct* FragmentType, FMassEntityHandle Entity)
+UE::Mass::Debug::FBreakpointHandle FMassDebugger::ShouldBreakOnFragmentWrite(const FMassEntityManager& EntityManager, const UScriptStruct* FragmentType, FMassEntityHandle Entity)
 {
-	if (LIKELY(!bHasBreakpoint))
+	if (LIKELY(!UE::Mass::Debug::FBreakpoint::HasBreakpoint()))
 	{
-		return false;
+		return UE::Mass::Debug::FBreakpointHandle::Invalid();
 	}
 
 	FEnvironment& ActiveEnvironment = GetActiveEnvironmentChecked(EntityManager);
 
 	if (LIKELY(!ActiveEnvironment.bHasBreakpoint))
 	{
-		return false;
+		return UE::Mass::Debug::FBreakpointHandle::Invalid();
 	}
 
-	const FMassEntityHandle* Found = ActiveEnvironment.FragmentWriteBreakpoints.FindPair(FragmentType, Entity);
+	const bool bFragmentHasBreakpoint = ActiveEnvironment.FragmentsWithBreakpoints.Contains(FragmentType);
 
-	if (Found != nullptr)
+	if (bFragmentHasBreakpoint)
 	{
-		return true;
+		UE::Mass::Debug::FBreakpointHandle FoundHandle;
+		for (UE::Mass::Debug::FBreakpoint& BreakPoint : ActiveEnvironment.Breakpoints)
+		{
+			if (BreakPoint.TriggerType == UE::Mass::Debug::FBreakpoint::ETriggerType::FragmentWrite)
+			{
+				TObjectKey<const UScriptStruct> BreakFragment = BreakPoint.Trigger.Get<TObjectKey<const UScriptStruct>>();
+				if (BreakFragment == FragmentType && BreakPoint.ApplyEntityFilter(EntityManager, Entity))
+				{
+					++BreakPoint.HitCount;
+					UE::Mass::Debug::FBreakpoint::LastBreakpointHandle = BreakPoint.Handle;
+					if (BreakPoint.bEnabled)
+					{
+						FoundHandle = BreakPoint.Handle;
+					}
+				}
+			}
+		}
+		return FoundHandle;
 	}
+	return UE::Mass::Debug::FBreakpointHandle::Invalid();
+}
 
-	if (Entity == ActiveEnvironment.SelectedEntity)
+UE::Mass::Debug::FBreakpoint* FMassDebugger::FindBreakpoint(const FMassEntityManager& EntityManager, UE::Mass::Debug::FBreakpointHandle Handle)
+{
+	FEnvironment* ActiveEnvironment = GetActiveEnvironment(EntityManager);
+	if (ActiveEnvironment)
 	{
-		if (ActiveEnvironment.SelectedEntityFragmentWriteBreakpoints.Contains(FragmentType))
-		{
-			return true;
-		}
+		return ActiveEnvironment->FindBreakpointByHandle(Handle);
 	}
+	return nullptr;
+}
 
-	return false;
+TArray<UE::Mass::Debug::FBreakpoint>& FMassDebugger::GetBreakpoints(const FMassEntityManager& EntityManager)
+{
+	FEnvironment& ActiveEnvironment = GetActiveEnvironmentChecked(EntityManager);
+	return ActiveEnvironment.Breakpoints;
 }
 
 bool FMassDebugger::HasAnyFragmentWriteBreakpoints(const FMassEntityManager& EntityManager, const UScriptStruct* FragmentType)
 {
-	if (LIKELY(!bHasBreakpoint))
+	if (LIKELY(!UE::Mass::Debug::FBreakpoint::HasBreakpoint()))
 	{
 		return false;
 	}
@@ -1029,131 +1101,194 @@ bool FMassDebugger::HasAnyFragmentWriteBreakpoints(const FMassEntityManager& Ent
 
 	if (FragmentType == nullptr)
 	{
-		return ActiveEnvironment.SelectedEntityFragmentWriteBreakpoints.Num() > 0 || ActiveEnvironment.FragmentWriteBreakpoints.Num() > 0;
+		return ActiveEnvironment.FragmentsWithBreakpoints.Num() > 0;
 	}
 
-	return ActiveEnvironment.SelectedEntityFragmentWriteBreakpoints.Contains(FragmentType) || ActiveEnvironment.FragmentWriteBreakpoints.Contains(FragmentType);
+	return ActiveEnvironment.FragmentsWithBreakpoints.Contains(FragmentType);
+}
+
+UE::Mass::Debug::FBreakpoint& FMassDebugger::CreateBreakpoint(const FMassEntityManager& EntityManager)
+{
+	FEnvironment& ActiveEnvironment = GetActiveEnvironmentChecked(EntityManager);
+
+	return ActiveEnvironment.Breakpoints.Emplace_GetRef();
 }
 
 void FMassDebugger::SetProcessorBreakpoint(const FMassEntityManager& EntityManager, TNotNull<const UMassProcessor*> Processor, FMassEntityHandle Entity)
 {
-	if (!FPlatformMisc::IsDebuggerPresent())
-	{
-		FMessageDialog::Open(EAppMsgType::Ok, LOCTEXT("NoDebuggerAttached", "Breakpoint set but no debugger is attached."));
-	}
+	UE::Mass::Debug::FNoDebuggerNotification::ConditionallyNotifyUser();
 
 	FEnvironment& ActiveEnvironment = GetActiveEnvironmentChecked(EntityManager);
 
 	ActiveEnvironment.bHasBreakpoint = true;
-	bHasBreakpoint = true;
+	UE::Mass::Debug::FBreakpoint::bHasBreakpoint = true;
+
+	UE::Mass::Debug::FBreakpoint& NewBreakpoint = ActiveEnvironment.Breakpoints.Emplace_GetRef();
+	NewBreakpoint.bEnabled = true;
+	NewBreakpoint.TriggerType = UE::Mass::Debug::FBreakpoint::ETriggerType::ProcessorExecute;
+	NewBreakpoint.Trigger.Set<TObjectKey<const UMassProcessor>>(Processor);
+	if (Entity.IsSet())
+	{
+		NewBreakpoint.FilterType = UE::Mass::Debug::FBreakpoint::EFilterType::SpecificEntity;
+		NewBreakpoint.Filter.Set<FMassEntityHandle>(Entity);
+	}
 
-	ActiveEnvironment.ProcessorBreakpoints.AddUnique(Processor, Entity);
+	ActiveEnvironment.ProcessorsWithBreakpoints.Add(Processor);
 	OnBreakpointsChangedDelegate.Broadcast();
 }
 
 void FMassDebugger::SetFragmentWriteBreakpoint(const FMassEntityManager& EntityManager, TNotNull<const UScriptStruct*> FragmentType, FMassEntityHandle Entity)
 {
-	if (!FPlatformMisc::IsDebuggerPresent())
-	{
-		FMessageDialog::Open(EAppMsgType::Ok, LOCTEXT("NoDebuggerAttached", "Breakpoint set but no debugger is attached."));
-	}
+	UE::Mass::Debug::FNoDebuggerNotification::ConditionallyNotifyUser();
 
 	FEnvironment& ActiveEnvironment = GetActiveEnvironmentChecked(EntityManager);
 
 	ActiveEnvironment.bHasBreakpoint = true;
-	bHasBreakpoint = true;
+	UE::Mass::Debug::FBreakpoint::bHasBreakpoint = true;
 
-	ActiveEnvironment.FragmentWriteBreakpoints.AddUnique(FragmentType, Entity);
+	UE::Mass::Debug::FBreakpoint& NewBreakpoint = ActiveEnvironment.Breakpoints.Emplace_GetRef();
+	NewBreakpoint.bEnabled = true;
+	NewBreakpoint.TriggerType = UE::Mass::Debug::FBreakpoint::ETriggerType::FragmentWrite;
+	NewBreakpoint.Trigger.Set<TObjectKey<const UScriptStruct>>(FragmentType);
+	if (Entity.IsSet())
+	{
+		NewBreakpoint.FilterType = UE::Mass::Debug::FBreakpoint::EFilterType::SpecificEntity;
+		NewBreakpoint.Filter.Set<FMassEntityHandle>(Entity);
+	}
+
+	ActiveEnvironment.FragmentsWithBreakpoints.Add(FragmentType);
 	OnBreakpointsChangedDelegate.Broadcast();
 }
 
-void FMassDebugger::ToggleFragmentWriteBreakForSelectedEntity(const FMassEntityManager& EntityManager, TNotNull<const UScriptStruct*> FragmentType)
+void FMassDebugger::SetFragmentWriteBreakForSelectedEntity(const FMassEntityManager& EntityManager, TNotNull<const UScriptStruct*> FragmentType)
 {
+	UE::Mass::Debug::FNoDebuggerNotification::ConditionallyNotifyUser();
+
 	FEnvironment& ActiveEnvironment = GetActiveEnvironmentChecked(EntityManager);
 
 	ActiveEnvironment.bHasBreakpoint = true;
-	bHasBreakpoint = true;
+	UE::Mass::Debug::FBreakpoint::bHasBreakpoint = true;
 
-	if (ActiveEnvironment.SelectedEntityFragmentWriteBreakpoints.Contains(FragmentType))
-	{
-		ActiveEnvironment.SelectedEntityFragmentWriteBreakpoints.Remove(FragmentType);
-	}
-	else
-	{
-		if (!FPlatformMisc::IsDebuggerPresent())
-		{
-			FMessageDialog::Open(EAppMsgType::Ok, LOCTEXT("NoDebuggerAttached", "Breakpoint set but no debugger is attached."));
-		}
-		ActiveEnvironment.SelectedEntityFragmentWriteBreakpoints.Add(FragmentType);
-	}
-		
+	UE::Mass::Debug::FBreakpoint NewBreakpoint = ActiveEnvironment.Breakpoints.Emplace_GetRef();
+	NewBreakpoint.bEnabled = true;
+	NewBreakpoint.TriggerType = UE::Mass::Debug::FBreakpoint::ETriggerType::FragmentWrite;
+	NewBreakpoint.Trigger.Set<TObjectKey<const UScriptStruct>>(FragmentType);
+	NewBreakpoint.FilterType = UE::Mass::Debug::FBreakpoint::EFilterType::SelectedEntity;
+
+	ActiveEnvironment.FragmentsWithBreakpoints.Add(FragmentType);
 	OnBreakpointsChangedDelegate.Broadcast();
 }
 
-bool FMassDebugger::IsFragmentWriteBreakForSelectedEntitySet(const FMassEntityManager& EntityManager, TNotNull<const UScriptStruct*> FragmentType)
+void FMassDebugger::SetBreakpointEnabled(UE::Mass::Debug::FBreakpointHandle Handle, bool bEnable)
 {
-	FEnvironment* ActiveEnvironment = GetActiveEnvironment(EntityManager);
-	return ActiveEnvironment ? ActiveEnvironment->SelectedEntityFragmentWriteBreakpoints.Contains(FragmentType) : false;
+	for (FEnvironment& Environment : ActiveEnvironments)
+	{
+		for (UE::Mass::Debug::FBreakpoint& Breakpoint : Environment.Breakpoints)
+		{
+			if (Breakpoint.Handle == Handle)
+			{
+				Breakpoint.bEnabled = bEnable;
+				OnBreakpointsChangedDelegate.Broadcast();
+				return;
+			}
+		}
+	}
 }
 
-
 void FMassDebugger::ClearProcessorBreakpoint(const FMassEntityManager& EntityManager, const UMassProcessor* Processor, FMassEntityHandle Entity)
 {
 	FEnvironment& ActiveEnvironment = GetActiveEnvironmentChecked(EntityManager);
 
-	ActiveEnvironment.ProcessorBreakpoints.Remove(Processor, Entity);
+	if (ActiveEnvironment.ProcessorsWithBreakpoints.Contains(Processor))
+	{
+		for (int EnvironmentIndex = 0; EnvironmentIndex < ActiveEnvironment.Breakpoints.Num(); EnvironmentIndex++)
+		{
+			UE::Mass::Debug::FBreakpoint& Breakpoint = ActiveEnvironment.Breakpoints[EnvironmentIndex];
 
-	UpdateHasBreakpoint();
+			if (Breakpoint.TriggerType == UE::Mass::Debug::FBreakpoint::ETriggerType::ProcessorExecute
+				&& Breakpoint.Trigger.Get<TObjectKey<const UMassProcessor>>() == Processor
+				&& Breakpoint.ApplyEntityFilter(EntityManager, Entity))
+			{
+				ActiveEnvironment.Breakpoints.RemoveAt(EnvironmentIndex);
+			}
+		}
+	}
 	OnBreakpointsChangedDelegate.Broadcast();
 }
 
 void FMassDebugger::ClearAllProcessorBreakpoints(const FMassEntityManager& EntityManager, const UMassProcessor* Processor)
 {
 	FEnvironment& ActiveEnvironment = GetActiveEnvironmentChecked(EntityManager);
-	ActiveEnvironment.ProcessorBreakpoints.Remove(Processor);
+	for (int EnvironmentIndex = 0; EnvironmentIndex < ActiveEnvironment.Breakpoints.Num(); EnvironmentIndex++)
+	{
+		UE::Mass::Debug::FBreakpoint& Breakpoint = ActiveEnvironment.Breakpoints[EnvironmentIndex];
+
+		if (Breakpoint.TriggerType == UE::Mass::Debug::FBreakpoint::ETriggerType::ProcessorExecute
+			&& Breakpoint.Trigger.Get<TObjectKey<const UMassProcessor>>() == Processor)
+		{
+			ActiveEnvironment.Breakpoints.RemoveAt(EnvironmentIndex);
+		}
+	}
+
+	ActiveEnvironment.ProcessorsWithBreakpoints.Remove(Processor);
 	OnBreakpointsChangedDelegate.Broadcast();
 }
 
 void FMassDebugger::ClearFragmentWriteBreak(const FMassEntityManager& EntityManager, const UScriptStruct* FragmentType, FMassEntityHandle Entity)
 {
 	FEnvironment& ActiveEnvironment = GetActiveEnvironmentChecked(EntityManager);
-	if (ActiveEnvironment.SelectedEntity == Entity)
+
+	if (ActiveEnvironment.FragmentsWithBreakpoints.Contains(FragmentType))
 	{
-		ActiveEnvironment.SelectedEntityFragmentWriteBreakpoints.Remove(FragmentType);
+		for (int EnvironmentIndex = 0; EnvironmentIndex < ActiveEnvironment.Breakpoints.Num(); EnvironmentIndex++)
+		{
+			UE::Mass::Debug::FBreakpoint& Breakpoint = ActiveEnvironment.Breakpoints[EnvironmentIndex];
+
+			if (Breakpoint.TriggerType == UE::Mass::Debug::FBreakpoint::ETriggerType::FragmentWrite
+				&& Breakpoint.Trigger.Get<TObjectKey<const UScriptStruct>>() == FragmentType
+				&& Breakpoint.ApplyEntityFilter(EntityManager, Entity))
+			{
+				ActiveEnvironment.Breakpoints.RemoveAt(EnvironmentIndex);
+			}
+		}
 	}
-	ActiveEnvironment.FragmentWriteBreakpoints.Remove(FragmentType, Entity);
-	UpdateHasBreakpoint();
 	OnBreakpointsChangedDelegate.Broadcast();
 }
 
 void FMassDebugger::ClearAllFragmentWriteBreak(const FMassEntityManager& EntityManager, const UScriptStruct* FragmentType)
 {
 	FEnvironment& ActiveEnvironment = GetActiveEnvironmentChecked(EntityManager);
-	ActiveEnvironment.FragmentWriteBreakpoints.Remove(FragmentType);
+	for (int EnvironmentIndex = 0; EnvironmentIndex < ActiveEnvironment.Breakpoints.Num(); EnvironmentIndex++)
+	{
+		UE::Mass::Debug::FBreakpoint& Breakpoint = ActiveEnvironment.Breakpoints[EnvironmentIndex];
+
+		if (Breakpoint.TriggerType == UE::Mass::Debug::FBreakpoint::ETriggerType::ProcessorExecute
+			&& Breakpoint.Trigger.Get<TObjectKey<const UScriptStruct>>() == FragmentType)
+		{
+			ActiveEnvironment.Breakpoints.RemoveAt(EnvironmentIndex);
+		}
+	}
+
+	ActiveEnvironment.FragmentsWithBreakpoints.Remove(FragmentType);
 	OnBreakpointsChangedDelegate.Broadcast();
 }
 
 void FMassDebugger::ClearAllEntityBreakpoints(const FMassEntityManager& EntityManager, FMassEntityHandle Entity)
 {
 	FEnvironment& ActiveEnvironment = GetActiveEnvironmentChecked(EntityManager);
-	TMultiMap<TObjectKey<const UMassProcessor>, FMassEntityHandle>::TIterator ProccessorIterator = ActiveEnvironment.ProcessorBreakpoints.CreateIterator();
 
-	for (; ProccessorIterator; ++ProccessorIterator)
+	for (int EnvironmentIndex = 0; EnvironmentIndex < ActiveEnvironment.Breakpoints.Num(); EnvironmentIndex++)
 	{
-		if (ProccessorIterator.Value() == Entity)
-		{
-			ProccessorIterator.RemoveCurrent();
-		}
-	}
+		UE::Mass::Debug::FBreakpoint& Breakpoint = ActiveEnvironment.Breakpoints[EnvironmentIndex];
 
-	TMultiMap<TObjectKey<const UScriptStruct>, FMassEntityHandle>::TIterator FragmentIterator = ActiveEnvironment.FragmentWriteBreakpoints.CreateIterator();
-	for (; FragmentIterator; ++FragmentIterator)
-	{
-		if (FragmentIterator.Value() == Entity)
+		FMassEntityHandle* BreakHandle = Breakpoint.Filter.TryGet<FMassEntityHandle>();
+		if (BreakHandle
+			&& *BreakHandle == Entity)
 		{
-			FragmentIterator.RemoveCurrent();
+			ActiveEnvironment.Breakpoints.RemoveAt(EnvironmentIndex);
 		}
 	}
+
 	OnBreakpointsChangedDelegate.Broadcast();
 }
  
@@ -1161,10 +1296,7 @@ void FMassDebugger::BreakOnFragmentWriteForSelectedEntity(FName FragmentName)
 {
 	for (FEnvironment& Environment : ActiveEnvironments)
 	{
-		if (Environment.EntityManager.Pin()->IsEntityValid(Environment.SelectedEntity))
-		{
-			SetFragmentWriteBreakpoint(*Environment.EntityManager.Pin(), GetFragmentTypeFromName(FragmentName), Environment.SelectedEntity);
-		}
+		SetFragmentWriteBreakForSelectedEntity(*Environment.EntityManager.Pin(), GetFragmentTypeFromName(FragmentName));
 	}
 	OnBreakpointsChangedDelegate.Broadcast();
 }
@@ -1175,13 +1307,34 @@ void FMassDebugger::ClearAllBreakpoints()
 	{
 		Environment.ClearBreakpoints();
 	}
-	bHasBreakpoint = false;
+	UE::Mass::Debug::FBreakpoint::bHasBreakpoint = false;
 	OnBreakpointsChangedDelegate.Broadcast();
 }
 
+void FMassDebugger::RemoveBreakpoint(UE::Mass::Debug::FBreakpointHandle Handle)
+{
+	for (FEnvironment& Environment : ActiveEnvironments)
+	{
+		for (int EnvironmentIndex = 0; EnvironmentIndex < Environment.Breakpoints.Num(); EnvironmentIndex++)
+		{
+			if (Environment.Breakpoints[EnvironmentIndex].Handle == Handle)
+			{
+				Environment.Breakpoints.RemoveAt(EnvironmentIndex);
+				RefreshBreakpoints();
+				return;
+			}
+		}
+	}
+}
+
 void FMassDebugger::ShutdownDebugger()
 {
-	ClearAllBreakpoints();
+	// Don't call ClearAllBreakpoints here because we don't want to send the BreakpointsChanged delegate
+	for (FEnvironment& Environment : ActiveEnvironments)
+	{
+		Environment.ClearBreakpoints();
+	}
+	UE::Mass::Debug::FBreakpoint::bHasBreakpoint = false;
 	ActiveEnvironments.Empty();
 	FragmentsByName.Empty();
 }
@@ -1361,16 +1514,33 @@ bool FMassDebugger::GetConstSharedFragmentData(const FMassEntityManager& EntityM
 	return false;
 }
 
-void FMassDebugger::UpdateHasBreakpoint()
+void FMassDebugger::RefreshBreakpoints()
 {
-	bHasBreakpoint = false;
+	UE::Mass::Debug::FBreakpoint::bHasBreakpoint = false;
+
 	for (FEnvironment& Environment : ActiveEnvironments)
 	{
-		Environment.bHasBreakpoint = Environment.ProcessorBreakpoints.Num() != 0
-			|| Environment.FragmentWriteBreakpoints.Num() != 0
-			|| Environment.SelectedEntityFragmentWriteBreakpoints.Num() != 0;
-		bHasBreakpoint |= Environment.bHasBreakpoint;
+		Environment.ProcessorsWithBreakpoints.Reset();
+		Environment.FragmentsWithBreakpoints.Reset();
+
+		for (UE::Mass::Debug::FBreakpoint& BreakPoint : Environment.Breakpoints)
+		{
+			if(TObjectKey<const UMassProcessor>* BreakProcessor = BreakPoint.Trigger.TryGet<TObjectKey<const UMassProcessor>>())
+			{
+				Environment.ProcessorsWithBreakpoints.Add(*BreakProcessor);
+			}
+			else if (TObjectKey<const UScriptStruct>* BreakFragment = BreakPoint.Trigger.TryGet<TObjectKey<const UScriptStruct>>())
+			{
+				Environment.FragmentsWithBreakpoints.Add(*BreakFragment);
+			}
+		}
+
+		Environment.bHasBreakpoint = Environment.ProcessorsWithBreakpoints.Num() != 0
+			|| Environment.FragmentsWithBreakpoints.Num() != 0;
+		UE::Mass::Debug::FBreakpoint::bHasBreakpoint |= Environment.bHasBreakpoint;
 	}
+
+	OnBreakpointsChangedDelegate.Broadcast();
 }
 
 FMassDebugger::FEnvironment& FMassDebugger::GetActiveEnvironmentChecked(const FMassEntityManager& EntityManager)
@@ -1427,10 +1597,22 @@ FMassDebugger::FEnvironment::~FEnvironment()
 
 void FMassDebugger::FEnvironment::ClearBreakpoints()
 {
-	SelectedEntityFragmentWriteBreakpoints.Reset();
-	ProcessorBreakpoints.Reset();
-	FragmentWriteBreakpoints.Reset();
-	bHasBreakpoint = false;
+	ProcessorsWithBreakpoints.Reset();
+	FragmentsWithBreakpoints.Reset();
+	Breakpoints.Reset();
+	UE::Mass::Debug::FBreakpoint::bHasBreakpoint = false;
+}
+
+UE::Mass::Debug::FBreakpoint* FMassDebugger::FEnvironment::FindBreakpointByHandle(UE::Mass::Debug::FBreakpointHandle Handle)
+{
+	for (UE::Mass::Debug::FBreakpoint& Breakpoint : Breakpoints)
+	{
+		if (Breakpoint.Handle == Handle)
+		{
+			return &Breakpoint;
+		}
+	}
+	return nullptr;
 }
 
 #undef LOCTEXT_NAMESPACE
diff --git a/Engine/Source/Runtime/MassEntity/Private/MassDebuggerBreakpoints.cpp b/Engine/Source/Runtime/MassEntity/Private/MassDebuggerBreakpoints.cpp
new file mode 100644
--- /dev/null
+++ b/Engine/Source/Runtime/MassEntity/Private/MassDebuggerBreakpoints.cpp
@@ -0,0 +1,577 @@
+// Copyright Epic Games, Inc. All Rights Reserved.
+
+#include "MassDebuggerBreakpoints.h"
+#if WITH_MASSENTITY_DEBUG
+#include "MassDebugger.h"
+#include "UObject/Class.h"
+#include "UObject/ObjectKey.h"
+#include "MassRequirements.h"
+#include "MassProcessor.h"
+
+#define LOCTEXT_NAMESPACE "MassDebugger"
+
+namespace UE::Mass::Debug
+{
+	UE_DISABLE_OPTIMIZATION_SHIP
+	void FBreakpoint::DebugBreak()
+	{
+		bool bDisableThisBreakpoint = false;
+
+		//====================================================================
+		//= A breakpoint set in the MassDebugger has triggered
+		//= Step out of this function to debug the actual code being run
+		// 
+		//= To disable this specific breakpoint use the Watch window to set
+		//= bDisableThisBreakpoint to `true` or 1
+		//====================================================================
+		UE_DEBUG_BREAK();
+
+		if (bDisableThisBreakpoint)
+		{
+			DisableActiveBreakpoint();
+		}
+	}
+	UE_ENABLE_OPTIMIZATION_SHIP
+
+	bool FBreakpoint::bHasBreakpoint = false;
+	FBreakpointHandle FBreakpoint::LastBreakpointHandle = FBreakpointHandle::Invalid();
+
+	FBreakpoint::FBreakpoint()
+	{
+		Handle = FBreakpointHandle::CreateHandle();
+	}
+
+	bool FBreakpoint::ApplyEntityFilterByFragments(const TArray<const UScriptStruct*>& Fragments) const
+	{
+		switch (FilterType)
+		{
+		case EFilterType::None:
+			return true;
+
+		case EFilterType::SpecificEntity:
+			return false;
+
+		case EFilterType::SelectedEntity:
+			return false;
+
+		case EFilterType::Query:
+			if (const FMassFragmentRequirements* Requirements = Filter.TryGet<FMassFragmentRequirements>())
+			{
+				TConstArrayView<FMassFragmentRequirementDescription> FragmentRequirements = Requirements->GetFragmentRequirements();
+				for (const FMassFragmentRequirementDescription& FragmentRequirement : FragmentRequirements)
+				{
+					if (FragmentRequirement.Presence == EMassFragmentPresence::All
+						&& !Fragments.Contains(FragmentRequirement.StructType))
+					{
+						return false;
+					}
+					else if (FragmentRequirement.Presence == EMassFragmentPresence::None
+						&& Fragments.Contains(FragmentRequirement.StructType))
+					{
+						return false;
+					}
+				}
+
+				TConstArrayView<FMassFragmentRequirementDescription> ChunkRequirements = Requirements->GetChunkFragmentRequirements();
+				for (const FMassFragmentRequirementDescription& FragmentRequirement : ChunkRequirements)
+				{
+					if (FragmentRequirement.Presence == EMassFragmentPresence::All
+						&& !Fragments.Contains(FragmentRequirement.StructType))
+					{
+						return false;
+					}
+					else if (FragmentRequirement.Presence == EMassFragmentPresence::None
+						&& Fragments.Contains(FragmentRequirement.StructType))
+					{
+						return false;
+					}
+				}
+
+				const FMassTagBitSet& RequiredAllTagsBitSet = Requirements->GetRequiredAllTags();
+				TArray<const UScriptStruct*> RequiredAllTags;
+				RequiredAllTagsBitSet.ExportTypes(RequiredAllTags);
+				for (const UScriptStruct* Tag : RequiredAllTags)
+				{
+					if (!Fragments.Contains(Tag))
+					{
+						return false;
+					}
+				}
+
+				const FMassTagBitSet& BlockedTagsBitSet = Requirements->GetRequiredNoneTags();
+				TArray<const UScriptStruct*> BlockedTags;
+				BlockedTagsBitSet.ExportTypes(BlockedTags);
+				for (const UScriptStruct* Tag : BlockedTags)
+				{
+					if (Fragments.Contains(Tag))
+					{
+						return false;
+					}
+				}
+			}
+			return false;
+		default:
+			break;
+		}
+
+		return true;
+	}
+
+	bool FBreakpoint::ApplyEntityFilterByArchetype(const FMassArchetypeHandle& ArchetypeHandle) const
+	{
+		switch (FilterType)
+		{
+		case EFilterType::None:
+			return true;
+
+		case EFilterType::Query:
+		{
+			if (const FMassFragmentRequirements* Requirements = Filter.TryGet<FMassFragmentRequirements>())
+			{
+				return Requirements->DoesArchetypeMatchRequirements(ArchetypeHandle);
+			}
+			return false;
+		}
+		default:
+			break;
+		}
+
+		return false;
+	}
+
+	bool FBreakpoint::ApplyEntityFilter(const FMassEntityManager& EntityManager, const FMassEntityHandle& Entity) const
+	{
+		switch (FilterType)
+		{
+		case EFilterType::None:
+			return true;
+
+		case EFilterType::SpecificEntity:
+		{
+			if (const FMassEntityHandle* SpecificEntity = Filter.TryGet<FMassEntityHandle>())
+			{
+				return SpecificEntity && (*SpecificEntity == Entity);
+			}
+			return false;
+		}
+
+		case EFilterType::SelectedEntity:
+		{
+			FMassEntityHandle SelectedEntity = FMassDebugger::GetSelectedEntity(EntityManager);
+			return SelectedEntity == Entity;
+		}
+
+		case EFilterType::Query:
+		{
+			if (const FMassFragmentRequirements* Requirements = Filter.TryGet<FMassFragmentRequirements>())
+			{
+				FMassArchetypeHandle Archetype = EntityManager.GetArchetypeForEntity(Entity);
+				return Requirements->DoesArchetypeMatchRequirements(Archetype);
+			}
+			return false;
+		}
+		default:
+			break;
+		}
+
+		return false;
+	}
+
+	FString FBreakpoint::TriggerTypeToString(FBreakpoint::ETriggerType InType)
+	{
+		switch (InType)
+		{
+		case FBreakpoint::ETriggerType::None:
+			return TEXT("None");
+
+		case FBreakpoint::ETriggerType::ProcessorExecute:
+			return TEXT("ProcessorExecute");
+
+		case FBreakpoint::ETriggerType::FragmentWrite:
+			return TEXT("FragmentWrite");
+
+		case FBreakpoint::ETriggerType::EntityCreate:
+			return TEXT("EntityCreate");
+
+		case FBreakpoint::ETriggerType::EntityDestroy:
+			return TEXT("EntityDestroy");
+
+		case FBreakpoint::ETriggerType::FragmentAdd:
+			return TEXT("FragmentAdd");
+
+		case FBreakpoint::ETriggerType::FragmentRemove:
+			return TEXT("FragmentRemove");
+
+		case FBreakpoint::ETriggerType::TagAdd:
+			return TEXT("TagAdd");
+
+		case FBreakpoint::ETriggerType::TagRemove:
+			return TEXT("TagRemove");
+		}
+		return TEXT("UnknownTrigger");
+	}
+
+	bool FBreakpoint::StringToTriggerType(const FString& InString, FBreakpoint::ETriggerType& OutType)
+	{
+		if (InString == TEXT("None"))
+		{
+			OutType = FBreakpoint::ETriggerType::None;
+			return true;
+		}
+		if (InString == TEXT("ProcessorExecute"))
+		{
+			OutType = FBreakpoint::ETriggerType::ProcessorExecute;
+			return true;
+		}
+		if (InString == TEXT("FragmentWrite"))
+		{
+			OutType = FBreakpoint::ETriggerType::FragmentWrite;
+			return true;
+		}
+		if (InString == TEXT("EntityCreate"))
+		{
+			OutType = FBreakpoint::ETriggerType::EntityCreate;
+			return true;
+		}
+		if (InString == TEXT("EntityDestroy"))
+		{
+			OutType = FBreakpoint::ETriggerType::EntityDestroy;
+			return true;
+		}
+		if (InString == TEXT("FragmentAdd"))
+		{
+			OutType = FBreakpoint::ETriggerType::FragmentAdd;
+			return true;
+		}
+		if (InString == TEXT("FragmentRemove"))
+		{
+			OutType = FBreakpoint::ETriggerType::FragmentRemove;
+			return true;
+		}
+		if (InString == TEXT("TagAdd"))
+		{
+			OutType = FBreakpoint::ETriggerType::TagAdd;
+			return true;
+		}
+		if (InString == TEXT("TagRemove"))
+		{
+			OutType = FBreakpoint::ETriggerType::TagRemove;
+			return true;
+		}
+
+		return false;
+	}
+
+
+	FString FBreakpoint::FilterTypeToString(FBreakpoint::EFilterType InType)
+	{
+		switch (InType)
+		{
+		case FBreakpoint::EFilterType::None:
+			return TEXT("None");
+
+		case FBreakpoint::EFilterType::SpecificEntity:
+			return TEXT("SpecificEntity");
+
+		case FBreakpoint::EFilterType::SelectedEntity:
+			return TEXT("SelectedEntity");
+
+		case FBreakpoint::EFilterType::Query:
+			return TEXT("Query");
+		}
+
+		return TEXT("UnknownFilter");
+	}
+
+	bool FBreakpoint::StringToFilterType(const FString& InString, FBreakpoint::EFilterType& OutType)
+	{
+		if (InString == TEXT("None"))
+		{
+			OutType = FBreakpoint::EFilterType::None;
+			return true;
+		}
+		if (InString == TEXT("SpecificEntity"))
+		{
+			OutType = FBreakpoint::EFilterType::SpecificEntity;
+			return true;
+		}
+		if (InString == TEXT("SelectedEntity"))
+		{
+			OutType = FBreakpoint::EFilterType::SelectedEntity;
+			return true;
+		}
+		if (InString == TEXT("Query"))
+		{
+			OutType = FBreakpoint::EFilterType::Query;
+			return true;
+		}
+
+		return false;
+	}
+
+	bool FBreakpoint::CheckDestroyEntityBreakpoints(const FMassEntityHandle& Entity)
+	{
+		if (LIKELY(!bHasBreakpoint))
+		{
+			return false;
+		}
+
+		for (const FMassDebugger::FEnvironment& Env : FMassDebugger::GetEnvironments())
+		{
+			if (!Env.bHasBreakpoint)
+			{
+				continue;
+			}
+
+			if (TSharedPtr<const FMassEntityManager> Manager = Env.EntityManager.Pin())
+			{
+				bool bShouldBreak = false;
+				for (const FBreakpoint& Breakpoint : Env.Breakpoints)
+				{
+					if (Breakpoint.TriggerType == FBreakpoint::ETriggerType::EntityDestroy
+						&& Breakpoint.ApplyEntityFilter(*Manager, Entity))
+					{
+						++Breakpoint.HitCount;
+						if (Breakpoint.bEnabled)
+						{
+							LastBreakpointHandle = Breakpoint.Handle;
+							// Set flag to break but check the rest of the breakpoints to keep accurate HitCounts
+							bShouldBreak = true;
+						}
+					}
+				}
+				if (bShouldBreak)
+				{
+					return true;
+				}
+			}
+		}
+		return false;
+	}
+
+	bool FBreakpoint::CheckDestroyEntityBreakpoints(TConstArrayView<FMassEntityHandle> Entities)
+	{
+		if (LIKELY(!bHasBreakpoint))
+		{
+			return false;
+		}
+
+		bool ShouldBreak = false;
+		for (const FMassEntityHandle& Handle : Entities)
+		{
+			// need to check them all before returning to ensure hitcount is accurate
+			ShouldBreak = CheckDestroyEntityBreakpoints(Handle) | ShouldBreak;
+		}
+		return ShouldBreak;
+	}
+
+	bool FBreakpoint::CheckFragmentAddBreakpoints(const FMassEntityHandle& Entity, const UScriptStruct* FragmentType)
+	{
+		if (LIKELY(!bHasBreakpoint))
+		{
+			return false;
+		}
+
+		for (const FMassDebugger::FEnvironment& Env : FMassDebugger::GetEnvironments())
+		{
+			if (!Env.bHasBreakpoint)
+			{
+				continue;
+			}
+			
+			if (TSharedPtr<const FMassEntityManager> ManagerPtr = Env.EntityManager.Pin())
+			{
+				if (Env.FragmentsWithBreakpoints.Contains(TObjectKey<const UScriptStruct>(FragmentType)))
+				{
+					bool bShouldBreak = false;
+
+					for (const FBreakpoint& Breakpoint : Env.Breakpoints)
+					{
+						if (Breakpoint.TriggerType != FBreakpoint::ETriggerType::FragmentAdd
+							|| !Breakpoint.Trigger.IsType<TObjectKey<const UScriptStruct>>())
+						{
+							continue;
+						}
+
+						const UScriptStruct* BreakpointFragmentType = Breakpoint.Trigger.Get<TObjectKey<const UScriptStruct>>().ResolveObjectPtr();
+
+						if (BreakpointFragmentType == FragmentType && Breakpoint.ApplyEntityFilter(*ManagerPtr, Entity))
+						{
+							++Breakpoint.HitCount;
+							if (Breakpoint.bEnabled)
+							{
+								LastBreakpointHandle = Breakpoint.Handle;
+								bShouldBreak = true;
+							}
+						}
+					}
+					if (bShouldBreak)
+					{
+						return true;
+					}
+				}
+			}
+		}
+		return false;
+	}
+
+	bool FBreakpoint::CheckFragmentAddBreakpoints(TConstArrayView<FMassEntityHandle> Entities, const UScriptStruct* FragmentType)
+	{
+		if (LIKELY(!bHasBreakpoint))
+		{
+			return false;
+		}
+
+		for (const FMassEntityHandle& Handle : Entities)
+		{
+			if (FBreakpoint::CheckFragmentAddBreakpoints(Handle, FragmentType))
+			{
+				return true;
+			}
+		}
+		return false;
+	}
+
+	bool FBreakpoint::CheckFragmentRemoveBreakpoints(const FMassEntityHandle& Handle, const UScriptStruct* FragmentType)
+	{
+		if (LIKELY(!bHasBreakpoint))
+		{
+			return false;
+		}
+
+		for (const FMassDebugger::FEnvironment& Env : FMassDebugger::GetEnvironments())
+		{
+			if (!Env.bHasBreakpoint)
+			{
+				continue;
+			}
+
+			if (TSharedPtr<const FMassEntityManager> ManagerPtr = Env.EntityManager.Pin())
+			{
+				if (Env.FragmentsWithBreakpoints.Contains(TObjectKey<const UScriptStruct>(FragmentType)))
+				{
+					bool bShouldBreak = false;
+
+					for (const FBreakpoint& Breakpoint : Env.Breakpoints)
+					{
+						if (Breakpoint.TriggerType != FBreakpoint::ETriggerType::FragmentRemove
+							|| !Breakpoint.Trigger.IsType<TObjectKey<const UScriptStruct>>())
+						{
+							continue;
+						}
+
+						const UScriptStruct* BreakpointFragmentType = Breakpoint.Trigger.Get<TObjectKey<const UScriptStruct>>().ResolveObjectPtr();
+
+						if (BreakpointFragmentType == FragmentType && Breakpoint.ApplyEntityFilter(*ManagerPtr, Handle))
+						{
+							++Breakpoint.HitCount;
+							if (Breakpoint.bEnabled)
+							{
+								LastBreakpointHandle = Breakpoint.Handle;
+								bShouldBreak = true;
+							}
+						}
+					}
+					if (bShouldBreak)
+					{
+						return true;
+					}
+				}
+			}
+		}
+		return false;
+	}
+
+	bool FBreakpoint::CheckCreateEntityBreakpoints(TArray<const UScriptStruct*> Fragments)
+	{
+		if (LIKELY(!bHasBreakpoint))
+		{
+			return false;
+		}
+
+		for (const FMassDebugger::FEnvironment& Env : FMassDebugger::GetEnvironments())
+		{
+			if (!Env.bHasBreakpoint)
+			{
+				continue;
+			}
+
+			bool bShouldBreak = false;
+
+			for (const FBreakpoint& Breakpoint : Env.Breakpoints)
+			{
+				if (Breakpoint.TriggerType != FBreakpoint::ETriggerType::EntityCreate)
+				{
+					continue;
+				}
+
+				if (Breakpoint.ApplyEntityFilterByFragments(Fragments))
+				{
+					++Breakpoint.HitCount;
+					if (Breakpoint.bEnabled)
+					{
+						LastBreakpointHandle = Breakpoint.Handle;
+						bShouldBreak = true;
+					}
+				}
+			}
+			if (bShouldBreak)
+			{
+				return true;
+			}
+		}
+
+		return false;
+	}
+
+	bool FBreakpoint::CheckCreateEntityBreakpoints(const FMassArchetypeHandle& ArchetypeHandle)
+	{
+		if (LIKELY(!bHasBreakpoint))
+		{
+			return false;
+		}
+
+		for (const FMassDebugger::FEnvironment& Env : FMassDebugger::GetEnvironments())
+		{
+			if (!Env.bHasBreakpoint)
+			{
+				continue;
+			}
+
+			bool bShouldBreak = false;
+
+			for (const FBreakpoint& Breakpoint : Env.Breakpoints)
+			{
+				if (Breakpoint.TriggerType != FBreakpoint::ETriggerType::EntityCreate)
+				{
+					continue;
+				}
+
+				if (Breakpoint.ApplyEntityFilterByArchetype(ArchetypeHandle))
+				{
+					++Breakpoint.HitCount;
+					if (Breakpoint.bEnabled)
+					{
+						LastBreakpointHandle = Breakpoint.Handle;
+						bShouldBreak = true;
+					}
+				}
+			}
+			if (bShouldBreak)
+			{
+				return true;
+			}
+		}
+
+		return false;
+	}
+
+	void FBreakpoint::DisableActiveBreakpoint()
+	{
+		FMassDebugger::SetBreakpointEnabled(LastBreakpointHandle, false);
+	}
+
+} // namespace UE::Mass::Debug
+
+#undef LOCTEXT_NAMESPACE
+#endif // WITH_MASSENTITY_DEBUG
\ No newline at end of file
diff --git a/Engine/Source/Runtime/MassEntity/Private/MassEntityManager.cpp b/Engine/Source/Runtime/MassEntity/Private/MassEntityManager.cpp
--- a/Engine/Source/Runtime/MassEntity/Private/MassEntityManager.cpp
+++ b/Engine/Source/Runtime/MassEntity/Private/MassEntityManager.cpp
@@ -233,13 +233,14 @@ void FMassEntityManager::Initialize(const FMassEntityManagerStorageInitParams& I
 			}
 		}
 	}
+
+	InitializationState = EInitializationState::Initialized;
+	bFirstCommandFlush = true;
+
 #if WITH_MASSENTITY_DEBUG
 	RequirementAccessDetector.Initialize();
 	FMassDebugger::RegisterEntityManager(*this);
 #endif // WITH_MASSENTITY_DEBUG
-
-	InitializationState = EInitializationState::Initialized;
-	bFirstCommandFlush = true;
 }
 
 void FMassEntityManager::PostInitialize()
@@ -583,6 +584,8 @@ FMassEntityHandle FMassEntityManager::CreateEntity(const FMassArchetypeHandle& A
 {
 	CHECK_SYNC_API_RETURN({});
 	check(ArchetypeHandle.IsValid());
+	
+	MASS_BREAKPOINT(UE::Mass::Debug::FBreakpoint::CheckCreateEntityBreakpoints(ArchetypeHandle));
 
 	const FMassEntityHandle Entity = ReserveEntity();
 	InternalBuildEntity(Entity
@@ -780,6 +783,8 @@ TSharedRef<FMassEntityManager::FEntityCreationContext> FMassEntityManager::Batch
 
 	TRACE_CPUPROFILER_EVENT_SCOPE(Mass_BatchCreateEntities);
 
+	MASS_BREAKPOINT(UE::Mass::Debug::FBreakpoint::CheckCreateEntityBreakpoints(ArchetypeHandle));
+
 	TConstArrayView<FMassEntityHandle> ReservedEntities = BatchReserveEntities(Count, InOutEntities);
 	
 	return InternalBatchCreateReservedEntities(
@@ -794,6 +799,8 @@ TSharedRef<FMassEntityManager::FEntityCreationContext> FMassEntityManager::Inter
 	FMassArchetypeData* ArchetypeData = FMassArchetypeHelper::ArchetypeDataFromHandle(ArchetypeHandle);
 	checkf(ArchetypeData, TEXT("Functions calling into this one are required to verify that the archetype handle is valid"));
 
+	MASS_BREAKPOINT(UE::Mass::Debug::FBreakpoint::CheckCreateEntityBreakpoints(ArchetypeHandle));
+
 	for (FMassEntityHandle Entity : ReservedEntities)
 	{
 		check(IsEntityValid(Entity));
@@ -819,6 +826,8 @@ void FMassEntityManager::DestroyEntity(FMassEntityHandle Entity)
 
 	FMassArchetypeData* Archetype = GetEntityStorageInterface().GetArchetype(Entity.Index);
 
+	MASS_BREAKPOINT(UE::Mass::Debug::FBreakpoint::CheckDestroyEntityBreakpoints(Entity));
+
 	if (Archetype)
 	{
 		ObserverManager.OnPreEntityDestroyed(Archetype->GetCompositionDescriptor(), Entity);
@@ -837,6 +846,8 @@ void FMassEntityManager::BatchDestroyEntities(TConstArrayView<FMassEntityHandle>
 
 	TRACE_CPUPROFILER_EVENT_SCOPE(Mass_BatchDestroyEntities);
 
+	MASS_BREAKPOINT(UE::Mass::Debug::FBreakpoint::CheckDestroyEntityBreakpoints(InEntities));
+
 	for (const FMassEntityHandle Entity : InEntities)
 	{
 		if (GetEntityStorageInterface().IsValidIndex(Entity.Index) == false)
diff --git a/Engine/Source/Runtime/MassEntity/Private/MassExecutionContext.cpp b/Engine/Source/Runtime/MassEntity/Private/MassExecutionContext.cpp
--- a/Engine/Source/Runtime/MassEntity/Private/MassExecutionContext.cpp
+++ b/Engine/Source/Runtime/MassEntity/Private/MassExecutionContext.cpp
@@ -233,44 +233,47 @@ void FMassExecutionContext::FEntityIterator::TestBreakpoints()
 {
 	FMassEntityManager& EntityManagerRef = ExecutionContext.GetEntityManagerChecked();
 	FMassEntityHandle Entity = ExecutionContext.GetEntity(EntityIndex);
-	if (QueryRuntime.bCheckProcessorBreaks && FMassDebugger::ShouldProcessorBreak(EntityManagerRef, ExecutionContext.DebugGetProcessor(), Entity))
+	if (QueryRuntime.bCheckProcessorBreaks)
 	{
-		bool bDisableThisBreakpoint = false;
-		//====================================================================
-		//= A breakpoint for this entity set in the MassDebugger has triggered
-		//= Step out of this function to debug the actual code being run for the entity
-		//=
-		//= To disable this specific breakpoint use the Watch window do set
-		//= bDisableThisBreakpoint to `true` or 1
-		//====================================================================
-		UE_DEBUG_BREAK();
-
-		if (bDisableThisBreakpoint)
+		if (UE::Mass::Debug::FBreakpointHandle BreakHandle = FMassDebugger::ShouldProcessorBreak(EntityManagerRef, ExecutionContext.DebugGetProcessor(), Entity))
 		{
-			FMassDebugger::ClearProcessorBreakpoint(EntityManagerRef, ExecutionContext.DebugGetProcessor(), Entity);
-		}
+			bool bDisableThisBreakpoint = false;
+			//====================================================================
+			//= A breakpoint for this entity set in the MassDebugger has triggered
+			//= Step out of this function to debug the actual code being run for the entity
+			//=
+			//= To disable this specific breakpoint use the Watch window to set
+			//= bDisableThisBreakpoint to `true` or 1
+			//====================================================================
+			UE_DEBUG_BREAK();
 
-		// bailing out, no point to hit multiple breakpoints for the given entity/processor pair
-		return;
+			if (bDisableThisBreakpoint)
+			{
+				FMassDebugger::SetBreakpointEnabled(BreakHandle, false);
+			}
+
+			// bailing out, no point to hit multiple breakpoints for the given entity/processor pair
+			return;
+		}
 	}
 
 	for (const UScriptStruct* Fragment : QueryRuntime.FragmentTypesToBreakOn)
 	{
-		if (FMassDebugger::ShouldBreakOnFragmentWrite(EntityManagerRef, Fragment, Entity))
+		if (UE::Mass::Debug::FBreakpointHandle BreakHandle = FMassDebugger::ShouldBreakOnFragmentWrite(EntityManagerRef, Fragment, Entity))
 		{
 			bool bDisableThisBreakpoint = false;
 			//====================================================================
 			//= A breakpoint for this entity set in the MassDebugger has triggered
 			//= Step out of this function to debug the actual code being run for the entity
 			// 
-			//= To disable this specific breakpoint use the Watch window do set
+			//= To disable this specific breakpoint use the Watch window to set
 			//= bDisableThisBreakpoint to `true` or 1
 			//====================================================================
 			UE_DEBUG_BREAK();
 
 			if (bDisableThisBreakpoint)
 			{
-				FMassDebugger::ClearFragmentWriteBreak(EntityManagerRef, Fragment, Entity);
+				FMassDebugger::SetBreakpointEnabled(BreakHandle, false);
 			}
 
 			// bailing out, no point to hit multiple breakpoints for the given entity/processor pair
diff --git a/Engine/Source/Runtime/MassEntity/Public/MassCommandBuffer.h b/Engine/Source/Runtime/MassEntity/Public/MassCommandBuffer.h
--- a/Engine/Source/Runtime/MassEntity/Public/MassCommandBuffer.h
+++ b/Engine/Source/Runtime/MassEntity/Public/MassCommandBuffer.h
@@ -9,9 +9,11 @@
 #include "Misc/MTTransactionallySafeAccessDetector.h"
 #include "Misc/TransactionallySafeCriticalSection.h"
 #include "MassEntityUtils.h"
+#include "MassDebuggerBreakpoints.h"
 #include "MassCommands.h"
 // required for UE::Mass::Deprecation
 #include <type_traits>
+#include "Concepts/ConvertibleTo.h"
 
 struct FMassEntityManager;
 
@@ -46,6 +48,52 @@ PRAGMA_ENABLE_DEPRECATION_WARNINGS
 #define ASSERT_EXECUTE_DEPRECATION(CommandType) \
 		static_assert(!UE::Mass::Deprecation::TOverridesExecute<CommandType>::value, "Mass Commands: CONST Execute function is deprecated in 5.7 and will be removed by 5.9. Use Run instead.");
 
+namespace UE::Mass::Debug
+{
+
+	template<typename TCommand, typename... Args>
+	concept HasCheckBreakpoints =
+		requires(Args&&... args) {
+			{ TCommand::CheckBreakpoints(Forward<Args>(args)...) }
+			-> CConvertibleTo<bool>;
+	};
+
+	template<typename TCommand, typename... Args>
+	concept HasCheckBreakpointsWithEntity =
+		requires(const FMassEntityHandle Entity, Args&&... args) {
+			{ TCommand::CheckBreakpoints(Entity, Forward<Args>(args)...) }
+			-> CConvertibleTo<bool>;
+	};
+
+	template<typename TCommand, typename... TArgs>
+	void CallCheckBreakpoints(TArgs&&... InArgs)
+	{
+#if WITH_MASSENTITY_DEBUG
+		if constexpr (HasCheckBreakpoints<TCommand, TArgs...>)
+		{
+			if (TCommand::CheckBreakpoints(Forward<TArgs>(InArgs)...))
+			{
+				UE::Mass::Debug::FBreakpoint::DebugBreak();
+			}
+		}
+#endif //WITH_MASSENTITY_DEBUG
+	}
+
+	template<typename TCommand, typename... TArgs >
+	void CallCheckBreakpointsByInstance(TArgs&&... InArgs)
+	{
+#if WITH_MASSENTITY_DEBUG
+		if constexpr (HasCheckBreakpointsWithEntity<TCommand, TArgs...>)
+		{
+			if (TCommand::CheckBreakpointsByInstance(Forward<TArgs>(InArgs)...))
+			{
+				UE::Mass::Debug::FBreakpoint::DebugBreak();
+			}
+		}
+#endif //WITH_MASSENTITY_DEBUG
+	}
+} // namespace UE::Mass::Debug
+
 
 struct FMassCommandBuffer
 {
@@ -57,6 +105,7 @@ struct FMassCommandBuffer
 	void PushCommand(const FMassEntityHandle Entity, TArgs&&... InArgs)
 	{
 		COMMAND_PUSHING_CHECK();
+		UE::Mass::Debug::CallCheckBreakpointsByInstance<TCommand<TArgs...>>(Entity, Forward<TArgs>(InArgs)...);
 
 		LLM_SCOPE_BYNAME(TEXT("Mass/PushCommand"));
 		TCommand<TArgs...>& Instance = CreateOrAddCommand<TCommand<TArgs...>>();
@@ -68,6 +117,7 @@ struct FMassCommandBuffer
 	void PushCommand(TArgs&&... InArgs)
 	{
 		COMMAND_PUSHING_CHECK();
+		UE::Mass::Debug::CallCheckBreakpointsByInstance<TCommand>(Forward<TArgs>(InArgs)...);
 
 		LLM_SCOPE_BYNAME(TEXT("Mass/PushCommand"));
 		TCommand& Instance = CreateOrAddCommand<TCommand>();
@@ -80,6 +130,7 @@ struct FMassCommandBuffer
 	void PushCommand(const FMassEntityHandle Entity)
 	{
 		COMMAND_PUSHING_CHECK();
+		UE::Mass::Debug::CallCheckBreakpoints<TCommand>(Entity);
 
 		LLM_SCOPE_BYNAME(TEXT("Mass/PushCommand"));
 		CreateOrAddCommand<TCommand>().Add(Entity);
@@ -90,6 +141,7 @@ struct FMassCommandBuffer
 	void PushCommand(TConstArrayView<FMassEntityHandle> Entities)
 	{
 		COMMAND_PUSHING_CHECK();
+		UE::Mass::Debug::CallCheckBreakpoints<TCommand>(Entities);
 
 		LLM_SCOPE_BYNAME(TEXT("Mass/PushCommand"));
 		CreateOrAddCommand<TCommand>().Add(Entities);
diff --git a/Engine/Source/Runtime/MassEntity/Public/MassCommands.h b/Engine/Source/Runtime/MassEntity/Public/MassCommands.h
--- a/Engine/Source/Runtime/MassEntity/Public/MassCommands.h
+++ b/Engine/Source/Runtime/MassEntity/Public/MassCommands.h
@@ -8,6 +8,7 @@
 #include "MassEntityTypes.h"
 #include "MassEntityUtils.h"
 #include "MassEntityManager.h"
+#include "MassDebuggerBreakpoints.h"
 #include "MassCommands.generated.h"
 
 /**
@@ -216,6 +217,14 @@ struct FMassCommandDestroyEntities : public FMassBatchedEntityCommand
 	{
 	}
 
+#if WITH_MASSENTITY_DEBUG
+	template<typename T>
+	static bool CheckBreakpoints(T Entity)
+	{
+		return UE::Mass::Debug::FBreakpoint::CheckDestroyEntityBreakpoints(Entity);
+	}
+#endif // WITH_MASSENTITY_DEBUG
+
 protected:
 	virtual void Run(FMassEntityManager& EntityManager) override
 	{
@@ -239,6 +248,14 @@ struct FMassCommandAddFragmentsInternal : public FMassBatchedEntityCommand
 		, FragmentsAffected(UE::Mass::Utils::ConstructFragmentBitSet<CheckTime, TTypes...>())
 	{}
 
+#if WITH_MASSENTITY_DEBUG
+	template<typename T>
+	static bool CheckBreakpoints(T Entity, TTypes... InFragments)
+	{
+		return UE::Mass::Debug::FBreakpoint::CheckFragmentAddBreakpoints(Entity, Forward<TTypes>(InFragments)...);
+	}
+#endif // WITH_MASSENTITY_DEBUG
+
 protected:
 	virtual void Run(FMassEntityManager& EntityManager) override
 	{
@@ -262,6 +279,14 @@ struct FMassCommandRemoveFragmentsInternal : public FMassBatchedEntityCommand
 		, FragmentsAffected(UE::Mass::Utils::ConstructFragmentBitSet<CheckTime, TTypes...>())
 	{}
 
+#if WITH_MASSENTITY_DEBUG
+	template<typename T>
+	static bool CheckBreakpoints(T Entity, TTypes... InFragments)
+	{
+		return UE::Mass::Debug::FBreakpoint::CheckFragmentRemoveBreakpoints(Entity, Forward<TTypes>(InFragments)...);
+	}
+#endif // WITH_MASSENTITY_DEBUG
+
 protected:
 	virtual void Run(FMassEntityManager& EntityManager) override
 	{
@@ -341,6 +366,14 @@ struct FMassCommandRemoveTagsInternal : public FMassCommandChangeTags
 			UE::Mass::Utils::ConstructTagBitSet<CheckTime, TTypes...>()
 			DEBUG_NAME("RemoveTags"))
 	{}
+
+#if WITH_MASSENTITY_DEBUG
+	template<typename T>
+	static bool CheckBreakpoints(T Entity, TTypes... InFragments)
+	{
+		return UE::Mass::Debug::FBreakpoint::CheckFragmentAddBreakpoints(Entity, Forward<TTypes>(InFragments)...);
+	}
+#endif // WITH_MASSENTITY_DEBUG
 };
 
 template<typename T>
@@ -384,6 +417,13 @@ struct FMassCommandAddFragmentInstances : public FMassBatchedEntityCommand
 		Fragments.Add(InFragments...);
 	}
 
+#if WITH_MASSENTITY_DEBUG
+	static bool CheckBreakpoints(const FMassEntityHandle Entity, TOthers... InFragments)
+	{
+		return UE::Mass::Debug::FBreakpoint::CheckFragmentAddBreakpoints(Entity, Forward<TOthers>(InFragments)...);
+	}
+#endif // WITH_MASSENTITY_DEBUG
+
 protected:
 	virtual void Reset() override
 	{
@@ -457,6 +497,13 @@ struct FMassCommandBuildEntity : public FMassCommandAddFragmentInstances<TOthers
 	{
 	}
 
+#if WITH_MASSENTITY_DEBUG
+	static bool CheckBreakpoints(TOthers... InFragments)
+	{
+		return UE::Mass::Debug::FBreakpoint::CheckCreateEntityBreakpoints(Forward<TOthers>(InFragments)...);
+	}
+#endif // WITH_MASSENTITY_DEBUG
+
 protected:
 
 	virtual void Run(FMassEntityManager& EntityManager) override
@@ -509,6 +556,14 @@ struct FMassCommandBuildEntityWithSharedFragments : public FMassBatchedCommand
 		bHasWork = true;
 	}
 
+#if WITH_MASSENTITY_DEBUG
+	static bool CheckBreakpoints(TOthers... InFragments)
+	{
+		// debugger doesn't currently support shared fragment filtering, so just send the others
+		return UE::Mass::Debug::FBreakpoint::CheckCreateEntityBreakpoints(Forward<TOthers>(InFragments)...);
+	}
+#endif // WITH_MASSENTITY_DEBUG
+
 protected:
 	virtual SIZE_T GetAllocatedSize() const override
 	{
diff --git a/Engine/Source/Runtime/MassEntity/Public/MassDebugger.h b/Engine/Source/Runtime/MassEntity/Public/MassDebugger.h
--- a/Engine/Source/Runtime/MassEntity/Public/MassDebugger.h
+++ b/Engine/Source/Runtime/MassEntity/Public/MassDebugger.h
@@ -6,6 +6,7 @@
 #include "MassEntityTrace.h"
 #if WITH_MASSENTITY_DEBUG
 #include "Containers/ContainersFwd.h"
+#include "MassDebuggerBreakpoints.h"
 #include "MassEntityQuery.h"
 #include "MassProcessor.h"
 #include "Async/TransactionallySafeMutex.h"
@@ -149,9 +150,13 @@ struct FMassDebugger
 
 		bool bHasBreakpoint = false;
 
-		TMultiMap<TObjectKey<const UMassProcessor>, FMassEntityHandle> ProcessorBreakpoints;
-		TMultiMap<TObjectKey<const UScriptStruct>, FMassEntityHandle> FragmentWriteBreakpoints;
-		TSet<TObjectKey<const UScriptStruct>> SelectedEntityFragmentWriteBreakpoints;
+		TArray<UE::Mass::Debug::FBreakpoint> Breakpoints;
+
+		MASSENTITY_API UE::Mass::Debug::FBreakpoint* FindBreakpointByHandle(UE::Mass::Debug::FBreakpointHandle Handle);
+
+		/** quick lookup to skip processors and fragments with no breakpoints set */
+		TSet<TObjectKey<const UMassProcessor>> ProcessorsWithBreakpoints;
+		TSet<TObjectKey<const UScriptStruct>> FragmentsWithBreakpoints;
 
 #if UE_MASS_TRACE_ENABLED
 		FDelegateHandle TraceStartedDelegateHandle;
@@ -268,9 +273,9 @@ struct FMassDebugger
 	static MASSENTITY_API bool DoesArchetypeMatchRequirements(const FMassArchetypeHandle& ArchetypeHandle, const FMassFragmentRequirements& Requirements, FOutputDevice& OutputDevice);
 
 	/**
-	 * Checks if a processor should break on execute for a given entity.
+	 * Checks if a processor should break on execute for a given entity. Returns the Id for the breakpoint if one is found, or 0.
 	 */
-	static MASSENTITY_API bool ShouldProcessorBreak(const FMassEntityManager& EntityManager, const UMassProcessor* Processor, FMassEntityHandle Entity);
+	static MASSENTITY_API UE::Mass::Debug::FBreakpointHandle ShouldProcessorBreak(const FMassEntityManager& EntityManager, const UMassProcessor* Processor, FMassEntityHandle Entity);
 
 	/**
 	 * Checks if a processor has any breakpoints set for any entity.
@@ -278,9 +283,24 @@ struct FMassDebugger
 	static MASSENTITY_API bool HasAnyProcessorBreakpoints(const FMassEntityManager& EntityManager, const UMassProcessor* Processor);
 
 	/**
-	 * Checks if a break should be triggered for a processor that's about to write a given fragment on an entity.
+	 * Checks if a break should be triggered for a processor that's about to write a given fragment on an entity. Returns the Id for the breakpoint if one is found (Invalid otherwise).
+	 */
+	static MASSENTITY_API UE::Mass::Debug::FBreakpointHandle ShouldBreakOnFragmentWrite(const FMassEntityManager& EntityManager, const UScriptStruct* FragmentType, FMassEntityHandle Entity);
+
+	/**
+	 * Find a breakpoint by Handle. Returns the pointer to the Breakpoint, or nullptr if none found.
 	 */
-	static MASSENTITY_API bool ShouldBreakOnFragmentWrite(const FMassEntityManager& EntityManager, const UScriptStruct* FragmentType, FMassEntityHandle Entity);
+	static MASSENTITY_API UE::Mass::Debug::FBreakpoint* FindBreakpoint(const FMassEntityManager& EntityManager, UE::Mass::Debug::FBreakpointHandle Handle);
+
+	/**
+	 * Get all the breakpoints for a given environment.
+	 */
+	static MASSENTITY_API TArray<UE::Mass::Debug::FBreakpoint>& GetBreakpoints(const FMassEntityManager& EntityManager);
+
+	/**
+	 * Refresh breakpoint flags after changes to breakpoint instances.
+	 */
+	static MASSENTITY_API void RefreshBreakpoints();
 
 	/**
 	 * Checks if there are any breakpoints set for writing a fragment for any entity
@@ -288,6 +308,11 @@ struct FMassDebugger
 	 */
 	static MASSENTITY_API bool HasAnyFragmentWriteBreakpoints(const FMassEntityManager& EntityManager, const UScriptStruct* FragmentType = nullptr);
 
+	/**
+	 * Create a default-constructed breakpoint
+	 */
+	static MASSENTITY_API UE::Mass::Debug::FBreakpoint& CreateBreakpoint(const FMassEntityManager& EntityManager);
+
 	/**
 	 * Sets a break to be triggered on processor execute for an entity.
 	 */
@@ -301,12 +326,12 @@ struct FMassDebugger
 	/**
 	 * Sets or removes a break before a write to a given fragment on whichever entity is selected at the time.
 	 */
-	static MASSENTITY_API void ToggleFragmentWriteBreakForSelectedEntity(const FMassEntityManager& EntityManager, TNotNull<const UScriptStruct*> FragmentType);
+	static MASSENTITY_API void SetFragmentWriteBreakForSelectedEntity(const FMassEntityManager& EntityManager, TNotNull<const UScriptStruct*> FragmentType);
 
 	/**
-	 * Returns if a break fragment write break is set for the selected entity.
+	 * Enable or Disable a breakpoint with a given Id.
 	 */
-	static MASSENTITY_API bool IsFragmentWriteBreakForSelectedEntitySet(const FMassEntityManager& EntityManager, TNotNull<const UScriptStruct*> FragmentType);
+	static MASSENTITY_API void SetBreakpointEnabled(UE::Mass::Debug::FBreakpointHandle Handle, bool bEnable);
 
 	/**
 	 * Clears a breakpoint triggered on processor execute for an entity.
@@ -384,6 +409,11 @@ struct FMassDebugger
 	 */
 	static MASSENTITY_API void ClearAllBreakpoints();
 
+	/**
+	 * Remove a single breakpoint by Id
+	 */
+	static MASSENTITY_API void RemoveBreakpoint(UE::Mass::Debug::FBreakpointHandle Handle);
+
 	/**
 	 * Destroy all debug environments and shutdown cleanly.
 	 */
@@ -393,10 +423,9 @@ private:
 	// @todo: contain static data in an instanced struct and setup/teardown in startup/shutdown module
 	static MASSENTITY_API TArray<FEnvironment> ActiveEnvironments;
 	static MASSENTITY_API UE::FTransactionallySafeMutex EntityManagerRegistrationLock;
-	static MASSENTITY_API bool bHasBreakpoint;
 	static MASSENTITY_API TMap<FName, const UScriptStruct*> FragmentsByName;
+	static MASSENTITY_API TArray<FName> CommandNames;
 
-	static MASSENTITY_API void UpdateHasBreakpoint();
 	static MASSENTITY_API FEnvironment& GetActiveEnvironmentChecked(const FMassEntityManager& EntityManager);
 	static MASSENTITY_API FEnvironment* GetActiveEnvironment(const FMassEntityManager& EntityManager);
 };
diff --git a/Engine/Source/Runtime/MassEntity/Public/MassDebuggerBreakpoints.h b/Engine/Source/Runtime/MassEntity/Public/MassDebuggerBreakpoints.h
new file mode 100644
--- /dev/null
+++ b/Engine/Source/Runtime/MassEntity/Public/MassDebuggerBreakpoints.h
@@ -0,0 +1,379 @@
+// Copyright Epic Games, Inc. All Rights Reserved.
+
+#pragma once
+
+#include "MassEntityMacros.h"
+
+#if WITH_MASSENTITY_DEBUG
+#define MASS_BREAKPOINT(bShouldBreak)					\
+	do {												\
+		if (UNLIKELY(bShouldBreak))						\
+		{												\
+			UE::Mass::Debug::FBreakpoint::DebugBreak();	\
+		}												\
+	} while (0)
+
+#include "MassEntityHandle.h"
+#include "Misc/TVariant.h"
+#include "UObject/ObjectKey.h"
+#include "MassRequirements.h"
+
+struct FMassDebugger;
+class UMassProcessor;
+class UScriptStruct;
+
+namespace UE::Mass::Debug
+{
+	struct FBreakpointHandle
+	{
+		int32 Handle = 0;
+
+		FBreakpointHandle(int32 InHandle)
+			: Handle(InHandle)
+		{
+		}
+
+		FBreakpointHandle() = default;
+
+		friend bool operator==(const FBreakpointHandle& A, const FBreakpointHandle& B)
+		{
+			return A.Handle == B.Handle;
+		}
+
+		bool IsValid() const
+		{
+			return Handle > 0;
+		}
+
+		operator bool() const
+		{
+			return IsValid();
+		}
+
+		static FBreakpointHandle CreateHandle()
+		{
+			static int32 LastHandle = 0;
+			return FBreakpointHandle{ ++LastHandle };
+		}
+
+		static FBreakpointHandle Invalid()
+		{
+			return FBreakpointHandle();
+		}
+	};
+
+	struct FBreakpoint
+	{
+		friend FMassDebugger;
+
+		MASSENTITY_API FBreakpoint();
+
+		enum class ETriggerType : uint8
+		{
+			None,
+			ProcessorExecute,
+			FragmentWrite,
+			EntityCreate,
+			EntityDestroy,
+			FragmentAdd,
+			FragmentRemove,
+			TagAdd,
+			TagRemove,
+			MAX
+		};
+		static MASSENTITY_API  FString TriggerTypeToString(FBreakpoint::ETriggerType InType);
+		static MASSENTITY_API bool StringToTriggerType(const FString& InString, FBreakpoint::ETriggerType& OutType);
+
+		enum class EFilterType : uint8
+		{
+			None,
+			SpecificEntity,
+			SelectedEntity,
+			Query
+		};
+
+		static MASSENTITY_API FString FilterTypeToString(FBreakpoint::EFilterType InType);
+		static MASSENTITY_API bool StringToFilterType(const FString& InString, FBreakpoint::EFilterType& OutType);
+
+		using TriggerVariant = TVariant<
+			TObjectKey<const UMassProcessor>,
+			TObjectKey<const UScriptStruct>
+		>;
+
+		using FilterVariant = TVariant<
+			FMassEntityHandle,
+			FMassFragmentRequirements
+		>;
+
+		FBreakpointHandle Handle;
+		mutable uint64 HitCount = 0;
+		ETriggerType TriggerType = ETriggerType::None;
+		EFilterType FilterType = EFilterType::None;
+		TriggerVariant Trigger;
+		FilterVariant Filter;
+		bool bEnabled = false;
+
+		/**
+		 * Apply the breakpoint Filter based on a list of fragment types
+		 *
+		 * @param Fragments	Array of fragment types to evaluate the breakpoint filter against
+		 *
+		 * @return True if the filter is set to an applicable type it evaluates to true for the provided Fragments
+		 */
+		MASSENTITY_API bool ApplyEntityFilterByFragments(const TArray<const UScriptStruct*>& Fragments) const;
+
+		/**
+		 * Apply the breakpoint Filter based on an FMassArchetypeHandle
+		 *
+		 * @param ArchetypeHandle	The archetype to evaluate breakpoint filter against
+		 *
+		 * @return True if the filter is set to an applicable type it evaluates to true for the provided ArchetypeHandle
+		 */
+		MASSENTITY_API bool ApplyEntityFilterByArchetype(const FMassArchetypeHandle& ArchetypeHandle) const;
+
+		/**
+		 * Evaluate the breakpoint filter against a specific entity
+		 *
+		 * @param EntityManager	The FMassEntityManager that owns the entity
+		 * @param Entity	The handle of the entity against which to evaluate the filter
+		 *
+		 * @return True if the filter is set to an applicable type it evaluates to true for the provided Entity
+		 */
+		MASSENTITY_API bool ApplyEntityFilter(const FMassEntityManager& EntityManager, const FMassEntityHandle& Entity) const;
+
+		/**
+		 * Checks if a fragment add on a given entity should trigger a breakpoint
+		 *
+		 * @param FMassEntityHandle	The entity being modified
+		 * @param FragmentType	The type of fragment being added
+		 *
+		 * @return True if any breakpoints are set and enabled that match this event
+		 */
+		static MASSENTITY_API bool CheckFragmentAddBreakpoints(const FMassEntityHandle& Handle, const UScriptStruct* FragmentType);
+
+		/**
+		 * Checks if a fragment add on a given entity should trigger a breakpoint (templated version)
+		 *
+		 * @param Handle	The entity being modified
+		 * @param InFragments	variadic list of fragment types being added
+		 *
+		 * @return True if any breakpoints are set and enabled that match this event
+		 */
+		template<typename... TFragments>
+		static bool CheckFragmentAddBreakpoints(const FMassEntityHandle& Handle, TFragments... InFragments)
+		{
+			if (LIKELY(!bHasBreakpoint))
+			{
+				return false;
+			}
+			return (... || CheckFragmentAddBreakpoints(Handle, TFragments::StaticStruct()));
+		}
+
+		/**
+		 * Checks if a fragment add on a group of entities should trigger a breakpoint
+		 *
+		 * @param Entities	The list of entities being modified
+		 * @param FragmentType	The type of fragment being added
+		 *
+		 * @return True if any breakpoints are set and enabled that match this event
+		 */
+		static MASSENTITY_API bool CheckFragmentAddBreakpoints(TConstArrayView<FMassEntityHandle> Entities, const UScriptStruct* FragmentType);
+
+		/**
+		 * Checks if a fragment add on a group of entities should trigger a breakpoint (templated version)
+		 *
+		 * @param Entities	The list of entities being modified
+		 * @param InFragments	variadic list of fragment types being added
+		 *
+		 * @return True if any breakpoints are set and enabled that match this event
+		 */
+		template<typename... TFragments>
+		static bool CheckFragmentAddBreakpoints(TConstArrayView<FMassEntityHandle> Entities, TFragments... InFragments)
+		{
+			if (LIKELY(!bHasBreakpoint))
+			{
+				return false;
+			}
+
+			return (... || CheckFragmentAddBreakpoints(Entities, TFragments::StaticStruct()));
+		}
+
+		/**
+		 * Checks if a fragment add on a group of entities should trigger a breakpoint (fragment instance version)
+		 *
+		 * @param Handle	The entity being modified
+		 * @param InFragments	variadic list of fragment instance references being added
+		 *
+		 * @return True if any breakpoints are set and enabled that match this event
+		 */
+		template<typename... TFragments>
+		static bool CheckFragmentAddBreakpointsByInstances(const FMassEntityHandle& Handle, TFragments... InFragments)
+		{
+			if (LIKELY(!bHasBreakpoint))
+			{
+				return false;
+			}
+			return (... || CheckFragmentAddBreakpoints(Handle, std::remove_reference_t<decltype(InFragments)>::StaticStruct()));
+		}
+
+		/**
+		 * Checks if a fragment add on a group of entities should trigger a breakpoint (fragment instance version)
+		 *
+		 * @param Entities	The list of entities being modified
+		 * @param InFragments	variadic list of fragment instance references being added
+		 *
+		 * @return True if any breakpoints are set and enabled that match this event
+		 */
+		template<typename... TFragments>
+		static bool CheckFragmentAddBreakpointsByInstances(TConstArrayView<FMassEntityHandle> Entities, TFragments... InFragments)
+		{
+			if (LIKELY(!bHasBreakpoint))
+			{
+				return false;
+			}
+
+			return (... || CheckFragmentAddBreakpoints(Entities, std::remove_reference_t<decltype(InFragments)>::StaticStruct()));
+		}
+
+		/**
+		 * Checks if a fragment remove on a given entity should trigger a breakpoint
+		 *
+		 * @param FMassEntityHandle	The entity being modified
+		 * @param FragmentType	The type of fragment being removed
+		 *
+		 * @return True if any breakpoints are set and enabled that match this event
+		 */
+		static MASSENTITY_API bool CheckFragmentRemoveBreakpoints(const FMassEntityHandle& Handle, const UScriptStruct* FragmentType);
+
+		/**
+		 * Checks if a fragment remove on a given entity should trigger a breakpoint (templated version)
+		 *
+		 * @param Handle	The entity being modified
+		 * @param InFragments	variadic list of fragment types being removed
+		 *
+		 * @return True if any breakpoints are set and enabled that match this event
+		 */
+		template<typename... TFragments>
+		static bool CheckFragmentRemoveBreakpoints(const FMassEntityHandle& Handle, TFragments... InFragments)
+		{
+			if (LIKELY(!bHasBreakpoint))
+			{
+				return false;
+			}
+			return (... || CheckFragmentRemoveBreakpoints(Handle, std::remove_reference_t<decltype(InFragments)>::StaticStruct()));
+		}
+
+		/**
+		 * Checks if a fragment remove on a group of entities should trigger a breakpoint
+		 *
+		 * @param Entities	The list of entities being modified
+		 * @param FragmentType	The type of fragment being removed
+		 *
+		 * @return True if any breakpoints are set and enabled that match this event
+		 */
+		static MASSENTITY_API bool CheckFragmentRemoveBreakpoints(TConstArrayView<FMassEntityHandle> Entities, const UScriptStruct* FragmentType);
+
+		/**
+		 * Checks if a fragment remove on a group of entities should trigger a breakpoint (templated version)
+		 *
+		 * @param Entities	The list of entities being modified
+		 * @param InFragments	variadic list of fragment types being removed
+		 *
+		 * @return True if any breakpoints are set and enabled that match this event
+		 */
+		template<typename... TFragments>
+		static bool CheckFragmentRemoveBreakpoints(TConstArrayView<FMassEntityHandle> Entities, TFragments... InFragments)
+		{
+			if (LIKELY(!bHasBreakpoint))
+			{
+				return false;
+			}
+			return (... || CheckFragmentRemoveBreakpoints(Entities, TFragments::StaticStruct()));
+		}
+
+		/**
+		 * Checks if an entity creation event should trigger a breakpoint
+		 *
+		 * @param Fragments	array of fragment types that compose the entity being created
+		 *
+		 * @return True if any breakpoints are set and enabled that match this event
+		 */
+		static bool CheckCreateEntityBreakpoints(TArray<const UScriptStruct*> Fragments);
+
+		/**
+		 * Checks if an entity creation event should trigger a breakpoint (templated version)
+		 *
+		 * @param InFragments	variadic list of fragment types that compose the entity being created
+		 *
+		 * @return True if any breakpoints are set and enabled that match this event
+		 */
+		template<typename... TFragments>
+		static bool CheckCreateEntityBreakpoints(TFragments... InFragments)
+		{
+			if (LIKELY(!bHasBreakpoint))
+			{
+				return false;
+			}
+
+			return CheckCreateEntityBreakpoints({ TFragments::StaticStruct()... });
+		}
+
+		/**
+		 * Checks if an entity creation event should trigger a breakpoint (archetype version)
+		 *
+		 * @param ArchetypeHandle	The archetype of the entity being created
+		 *
+		 * @return True if any breakpoints are set and enabled that match this event
+		 */
+		static MASSENTITY_API bool CheckCreateEntityBreakpoints(const FMassArchetypeHandle& ArchetypeHandle);
+
+		/**
+		 * Checks if an entity destruction event should trigger a breakpoint
+		 *
+		 * @param Entity	The entity being destroyed
+		 *
+		 * @return True if any breakpoints are set and enabled that match this event
+		 */
+		static MASSENTITY_API bool CheckDestroyEntityBreakpoints(const FMassEntityHandle& Entity);
+
+		/**
+		 * Checks if an entity destruction event should trigger a breakpoint
+		 *
+		 * @param Entities	The array entities being destroyed
+		 *
+		 * @return True if any breakpoints are set and enabled that match this event
+		 */
+		static MASSENTITY_API bool CheckDestroyEntityBreakpoints(TConstArrayView<FMassEntityHandle> Entities);
+
+		/**
+		 * Disable the currently-active breakpoint (if any)
+		 */
+		static MASSENTITY_API void DisableActiveBreakpoint();
+
+		/**
+		 * Checks if any breakpoints are currently set
+		 *
+		 * @return True if there are any breakpoints set in the mass debugger even if they're disabled (so the hit count can still increment).
+		 */
+		static MASSENTITY_API inline bool HasBreakpoint()
+		{
+			return bHasBreakpoint;
+		}
+
+		/**
+		 * Triggers a breakpoint in the attached debugger (if present)
+		 */
+		static MASSENTITY_API void DebugBreak();
+
+	private:
+		static bool bHasBreakpoint;
+		static FBreakpointHandle LastBreakpointHandle;
+	};
+	
+} // namespace UE::Mass::Debug
+
+#else // WITH_MASSENTITY_DEBUG
+#define MASS_BREAKPOINT(...)	\
+		do { } while (0)
+
+#endif // WITH_MASSENTITY_DEBUG
