[Mass Debugger]
Properly guard debugger usage for standalone editor builds

diff --git a/Engine/Source/Editor/MassEntityDebugger/Private/MassDebuggerModel.cpp b/Engine/Source/Editor/MassEntityDebugger/Private/MassDebuggerModel.cpp
--- a/Engine/Source/Editor/MassEntityDebugger/Private/MassDebuggerModel.cpp
+++ b/Engine/Source/Editor/MassEntityDebugger/Private/MassDebuggerModel.cpp
@@ -823,11 +823,12 @@ void FMassDebuggerModel::LoadBreakpointsFromDisk()
 		}
 	}
 	OnEditorBreakpointsChangedDelegate.Broadcast();
-#endif
+#endif //WITH_MASSENTITY_DEBUG
 }
 
 void FMassDebuggerModel::SaveBreakpointsToDisk()
 {
+#if WITH_MASSENTITY_DEBUG
 	const FString FileName = TEXT("Breakpoints.json");
 	const FString SavedDir = FPaths::ProjectSavedDir() / TEXT("MassDebugger");
 	IFileManager::Get().MakeDirectory(*SavedDir, /*Tree=*/true);
@@ -845,10 +846,12 @@ void FMassDebuggerModel::SaveBreakpointsToDisk()
 	const TSharedRef<TJsonWriter<>> Writer = TJsonWriterFactory<>::Create(&OutString);
 	FJsonSerializer::Serialize(RootObject.ToSharedRef(), Writer);
 	FFileHelper::SaveStringToFile(OutString, *FullPath);
+#endif //WITH_MASSENTITY_DEBUG
 }
 
 void FMassDebuggerModel::CreateBreakpoint()
 {
+#if WITH_MASSENTITY_DEBUG
 	TSharedPtr<FMassDebuggerBreakpointData> NewBp = MakeShared<FMassDebuggerBreakpointData>();
 	CachedBreakpoints.Add(NewBp);
 	if (!IsStale())
@@ -856,10 +859,12 @@ void FMassDebuggerModel::CreateBreakpoint()
 		NewBp->ApplyToEngine(*this);
 	}
 	OnEditorBreakpointsChangedDelegate.Broadcast();
+#endif //WITH_MASSENTITY_DEBUG
 }
 
 void FMassDebuggerModel::CreateBreakpointFromString(const FString& InString)
 {
+#if WITH_MASSENTITY_DEBUG
 	TSharedPtr<FMassDebuggerBreakpointData> NewBp = MakeShared<FMassDebuggerBreakpointData>();
 	if (NewBp->DeserializeFromJsonString(InString))
 	{
@@ -870,10 +875,12 @@ void FMassDebuggerModel::CreateBreakpointFromString(const FString& InString)
 		}
 		OnEditorBreakpointsChangedDelegate.Broadcast();
 	}
+#endif //WITH_MASSENTITY_DEBUG
 }
 
 void FMassDebuggerModel::ApplyBreakpointsToCurrentEnvironment()
 {
+#if WITH_MASSENTITY_DEBUG
 	if (!IsStale())
 	{
 		FMassDebugger::ClearAllBreakpoints();
@@ -886,25 +893,31 @@ void FMassDebuggerModel::ApplyBreakpointsToCurrentEnvironment()
 		FMassDebugger::RefreshBreakpoints();
 	}
 	OnEditorBreakpointsChangedDelegate.Broadcast();
+#endif //WITH_MASSENTITY_DEBUG
 }
 
 void FMassDebuggerModel::ReconcileAllBreakpoints()
 {
+#if WITH_MASSENTITY_DEBUG
 	for (TSharedPtr<FMassDebuggerBreakpointData>& Breakpoint : CachedBreakpoints)
 	{
 		Breakpoint->ReconcileDataFromNames(*this);
 	}
 	OnEditorBreakpointsChangedDelegate.Broadcast();
+#endif //WITH_MASSENTITY_DEBUG
 }
 
 void FMassDebuggerModel::RemoveAllBreakpoints()
 {
+#if WITH_MASSENTITY_DEBUG
 	CachedBreakpoints.Reset();
 	ApplyBreakpointsToCurrentEnvironment();
+#endif //WITH_MASSENTITY_DEBUG
 }
 
 void FMassDebuggerModel::RemoveBreakpoint(UE::Mass::Debug::FBreakpointHandle Handle)
 {
+#if WITH_MASSENTITY_DEBUG
 	for (int32 i = 0; i < CachedBreakpoints.Num(); i++)
 	{
 		if (CachedBreakpoints[i]->BreakpointInstance.Handle == Handle)
@@ -914,16 +927,20 @@ void FMassDebuggerModel::RemoveBreakpoint(UE::Mass::Debug::FBreakpointHandle Han
 			return;
 		}
 	}
+#endif //WITH_MASSENTITY_DEBUG
 }
 
 FMassDebuggerBreakpointData::FMassDebuggerBreakpointData(const UE::Mass::Debug::FBreakpoint& InBreakpoint, FMassDebuggerModel& Model)
 {
+#if WITH_MASSENTITY_DEBUG
 	BreakpointInstance = InBreakpoint;
 	ReconcileDataFromEngineBreakpoint(Model);
+#endif //WITH_MASSENTITY_DEBUG
 }
 
 void FMassDebuggerBreakpointData::ReconcileDataFromNames(FMassDebuggerModel& Model)
 {
+#if WITH_MASSENTITY_DEBUG
 	switch (BreakpointInstance.TriggerType)
 	{
 	case UE::Mass::Debug::FBreakpoint::ETriggerType::ProcessorExecute:
@@ -974,10 +991,12 @@ void FMassDebuggerBreakpointData::ReconcileDataFromNames(FMassDebuggerModel& Mod
 			}
 		}
 	}
+#endif //WITH_MASSENTITY_DEBUG
 }
 
 void FMassDebuggerBreakpointData::ReconcileDataFromEngineBreakpoint(FMassDebuggerModel& Model)
 {
+#if WITH_MASSENTITY_DEBUG
 	switch (BreakpointInstance.TriggerType)
 	{
 	case UE::Mass::Debug::FBreakpoint::ETriggerType::ProcessorExecute:
@@ -1010,6 +1029,7 @@ void FMassDebuggerBreakpointData::ReconcileDataFromEngineBreakpoint(FMassDebugge
 	default:
 		break;
 	}
+#endif //WITH_MASSENTITY_DEBUG
 }
 
 void FMassDebuggerBreakpointData::ApplyToEngine(FMassDebuggerModel& Model, bool bRefreshEngineBreakpoints/* = true*/)
@@ -1051,6 +1071,7 @@ TSharedPtr<FJsonObject> FMassDebuggerBreakpointData::SerializeToJson() const
 {
 	TSharedPtr<FJsonObject> JsonObject = MakeShared<FJsonObject>();
 
+#if WITH_MASSENTITY_DEBUG
 	JsonObject->SetBoolField(TEXT("Enabled"), BreakpointInstance.bEnabled);
 	JsonObject->SetStringField(TEXT("TriggerType"), UE::Mass::Debug::FBreakpoint::TriggerTypeToString(BreakpointInstance.TriggerType));
 	JsonObject->SetStringField(TEXT("FilterType"),UE::Mass::Debug::FBreakpoint::FilterTypeToString(BreakpointInstance.FilterType));
@@ -1110,12 +1131,13 @@ TSharedPtr<FJsonObject> FMassDebuggerBreakpointData::SerializeToJson() const
 	default:
 		break;
 	}
-
+#endif //WITH_MASSENTITY_DEBUG
 	return JsonObject;
 }
 
 bool FMassDebuggerBreakpointData::DeserializeFromJson(const TSharedPtr<FJsonObject>& JsonObject)
 {
+#if WITH_MASSENTITY_DEBUG
 	if (!JsonObject.IsValid())
 	{
 		return false;
@@ -1193,10 +1215,14 @@ bool FMassDebuggerBreakpointData::DeserializeFromJson(const TSharedPtr<FJsonObje
 	}
 
 	return true;
+#else
+	return false;
+#endif //WITH_MASSENTITY_DEBUG
 }
 
 bool FMassDebuggerBreakpointData::DeserializeFromJsonString(const FString& JsonString)
 {
+#if WITH_MASSENTITY_DEBUG
 	TSharedPtr<FJsonObject> JsonObject;
 	TSharedRef<TJsonReader<>> Reader = TJsonReaderFactory<>::Create(JsonString);
 
@@ -1206,6 +1232,9 @@ bool FMassDebuggerBreakpointData::DeserializeFromJsonString(const FString& JsonS
 	}
 
 	return DeserializeFromJson(JsonObject);
+#else
+	return false;
+#endif //WITH_MASSENTITY_DEBUG
 }
 
 
@@ -1476,6 +1505,7 @@ void FMassDebuggerModel::OnEntitySelected(const FMassEntityManager& EntityManage
 
 void FMassDebuggerModel::OnBreakpointsChanged()
 {
+#if WITH_MASSENTITY_DEBUG
 	if (IsStale())
 	{
 		return;
@@ -1506,6 +1536,7 @@ void FMassDebuggerModel::OnBreakpointsChanged()
 		}
 		OnEditorBreakpointsChangedDelegate.Broadcast();
 	}
+#endif //WITH_MASSENTITY_DEBUG
 }
 
 void FMassDebuggerModel::SelectFragment(FName InFragementName)
diff --git a/Engine/Source/Editor/MassEntityDebugger/Private/SMassBreakpointsView.cpp b/Engine/Source/Editor/MassEntityDebugger/Private/SMassBreakpointsView.cpp
--- a/Engine/Source/Editor/MassEntityDebugger/Private/SMassBreakpointsView.cpp
+++ b/Engine/Source/Editor/MassEntityDebugger/Private/SMassBreakpointsView.cpp
@@ -232,6 +232,7 @@ void SMassBreakpointsView::SBreakpointTableRow::Construct(const FArguments& InAr
 
 TSharedRef<SWidget> SMassBreakpointsView::SBreakpointTableRow::GenerateWidgetForColumn(const FName& ColumnId)
 {
+#if WITH_MASSENTITY_DEBUG
 	TSharedPtr<FMassDebuggerBreakpointData> Data = BreakpointData.Pin();
 	TSharedPtr<SMassBreakpointsView> View = ParentView.Pin();
 
@@ -597,7 +598,7 @@ TSharedRef<SWidget> SMassBreakpointsView::SBreakpointTableRow::GenerateWidgetFor
 			return FText::AsNumber(0);
 		});
 	}
-
+#endif //WITH_MASSENTITY_DEBUG
     return SNullWidget::NullWidget;
 }
 
diff --git a/Engine/Source/Editor/MassEntityDebugger/Private/SMassBreakpointsView.h b/Engine/Source/Editor/MassEntityDebugger/Private/SMassBreakpointsView.h
--- a/Engine/Source/Editor/MassEntityDebugger/Private/SMassBreakpointsView.h
+++ b/Engine/Source/Editor/MassEntityDebugger/Private/SMassBreakpointsView.h
@@ -7,6 +7,7 @@
 #include "Widgets/Input/SCheckBox.h"
 #include "SSearchableComboBox.h"
 #include "MassDebugger.h"
+#include "MassDebuggerBreakpoints.h"
 
 class SSearchableComboBox;
 struct FMassDebuggerModel;
diff --git a/Engine/Source/Runtime/MassEntity/Public/MassDebuggerBreakpoints.h b/Engine/Source/Runtime/MassEntity/Public/MassDebuggerBreakpoints.h
--- a/Engine/Source/Runtime/MassEntity/Public/MassDebuggerBreakpoints.h
+++ b/Engine/Source/Runtime/MassEntity/Public/MassDebuggerBreakpoints.h
@@ -373,6 +373,23 @@ namespace UE::Mass::Debug
 } // namespace UE::Mass::Debug
 
 #else // WITH_MASSENTITY_DEBUG
+namespace UE::Mass::Debug
+{
+	struct FBreakpointHandle
+	{
+	};
+
+	struct FBreakpoint
+	{
+		enum class ETriggerType : uint8
+		{
+		};
+		enum class EFilterType : uint8
+		{
+		};
+	};
+}
+
 #define MASS_BREAKPOINT(...)	\
 		do { } while (0)
 
