[Mass Debugger]
Fix crash when creating breakpoints for shared or chunk fragment writes.

diff --git a/Engine/Source/Editor/MassEntityDebugger/Private/MassDebuggerModel.cpp b/Engine/Source/Editor/MassEntityDebugger/Private/MassDebuggerModel.cpp
--- a/Engine/Source/Editor/MassEntityDebugger/Private/MassDebuggerModel.cpp
+++ b/Engine/Source/Editor/MassEntityDebugger/Private/MassDebuggerModel.cpp
@@ -664,10 +664,14 @@ void FMassDebuggerModel::CacheProcessingGraphs(TConstArrayView<TNotNull<const UM
 
 void FMassDebuggerModel::CacheFragmentData(
 	TArray<TSharedPtr<FMassDebuggerFragmentData>>& OutData,
-	const TConstArrayView<TWeakObjectPtr<const UScriptStruct>>& FragmentTypes
+	const TConstArrayView<TWeakObjectPtr<const UScriptStruct>>& FragmentTypes,
+	bool bAppend
 )
 {
-	OutData.Reset();
+	if (!bAppend)
+	{
+		OutData.Reset();
+	}
 
 #if WITH_MASSENTITY_DEBUG
 	TMap<const UScriptStruct*, TSharedPtr<FMassDebuggerFragmentData>> TempFragmentDataMap;
@@ -729,7 +733,10 @@ void FMassDebuggerModel::CacheFragmentData(
 void FMassDebuggerModel::CacheFragmentData()
 {
 #if WITH_MASSENTITY_DEBUG
-	CacheFragmentData(CachedFragmentData, FMassFragmentBitSet::DebugGetAllStructTypes());
+	CacheFragmentData(CachedFragmentData, FMassFragmentBitSet::DebugGetAllStructTypes(), false);
+	CacheFragmentData(CachedFragmentData, FMassChunkFragmentBitSet::DebugGetAllStructTypes(), true);
+	CacheFragmentData(CachedFragmentData, FMassSharedFragmentBitSet::DebugGetAllStructTypes(), true);
+	CacheFragmentData(CachedFragmentData, FMassConstSharedFragmentBitSet::DebugGetAllStructTypes(), true);
 	FragmentNames.Reset();
 	FragmentMap.Reset();
 
@@ -758,7 +765,7 @@ TSharedPtr<FMassDebuggerFragmentData> FMassDebuggerModel::FindFragmentData(const
 void FMassDebuggerModel::CacheTagData()
 {
 #if WITH_MASSENTITY_DEBUG
-	CacheFragmentData(CachedTagData, FMassTagBitSet::DebugGetAllStructTypes());
+	CacheFragmentData(CachedTagData, FMassTagBitSet::DebugGetAllStructTypes(), false);
 	TagNames.Reset();
 	TagMap.Reset();
 
@@ -1021,8 +1028,13 @@ void FMassDebuggerBreakpointData::ReconcileDataFromEngineBreakpoint(FMassDebugge
 		if (BreakpointInstance.Trigger.IsType<TObjectKey<const UScriptStruct>>())
 		{
 			const UScriptStruct* StructType = BreakpointInstance.Trigger.Get<TObjectKey<const UScriptStruct>>().ResolveObjectPtr();
-			TriggerFragment = Model.FindFragmentData(StructType);
-			TriggerName = TriggerFragment->Name.ToString();
+			if (StructType)
+			{
+				if ((TriggerFragment = Model.FindFragmentData(StructType)))
+				{
+					TriggerName = TriggerFragment->Name.ToString();
+				}
+			}
 		}
 		break;
 
diff --git a/Engine/Source/Editor/MassEntityDebugger/Private/MassDebuggerModel.h b/Engine/Source/Editor/MassEntityDebugger/Private/MassDebuggerModel.h
--- a/Engine/Source/Editor/MassEntityDebugger/Private/MassDebuggerModel.h
+++ b/Engine/Source/Editor/MassEntityDebugger/Private/MassDebuggerModel.h
@@ -319,7 +319,8 @@ struct FMassDebuggerModel
 
 	void CacheFragmentData(
 		TArray<TSharedPtr<FMassDebuggerFragmentData>>& OutData,
-		const TConstArrayView<TWeakObjectPtr<const UScriptStruct>>& FragmentTypes
+		const TConstArrayView<TWeakObjectPtr<const UScriptStruct>>& FragmentTypes,
+		bool bAppend
 	);
 	void CacheFragmentData();
 	TSharedPtr<FMassDebuggerFragmentData> FindFragmentData(const UScriptStruct* FragmentType);
