[Mass Debugger]
Fix crashes in entity destruction breakpoints with query-based filtering

diff --git a/Engine/Source/Editor/MassEntityDebugger/Private/MassDebuggerModel.cpp b/Engine/Source/Editor/MassEntityDebugger/Private/MassDebuggerModel.cpp
--- a/Engine/Source/Editor/MassEntityDebugger/Private/MassDebuggerModel.cpp
+++ b/Engine/Source/Editor/MassEntityDebugger/Private/MassDebuggerModel.cpp
@@ -1057,7 +1057,8 @@ void FMassDebuggerBreakpointData::ApplyToEngine(FMassDebuggerModel& Model, bool
 	if (FMassDebugger::FEnvironment* Environment = FMassDebugger::FindEnvironmentForEntityManager(*EntityManager))
 	{
 		if (BreakpointInstance.FilterType == UE::Mass::Debug::FBreakpoint::EFilterType::Query
-			&& !BreakpointInstance.Filter.IsType<FMassFragmentRequirements>())
+			&& !BreakpointInstance.Filter.IsType<FMassFragmentRequirements>()
+			&& FilterQuery.IsValid())
 		{
 			BreakpointInstance.Filter.Set<FMassFragmentRequirements>(FilterQuery->BuildEntityQuery(EntityManager.ToSharedRef()));
 		}
diff --git a/Engine/Source/Runtime/MassEntity/Private/MassDebuggerBreakpoints.cpp b/Engine/Source/Runtime/MassEntity/Private/MassDebuggerBreakpoints.cpp
--- a/Engine/Source/Runtime/MassEntity/Private/MassDebuggerBreakpoints.cpp
+++ b/Engine/Source/Runtime/MassEntity/Private/MassDebuggerBreakpoints.cpp
@@ -166,7 +166,10 @@ namespace UE::Mass::Debug
 			if (const FMassFragmentRequirements* Requirements = Filter.TryGet<FMassFragmentRequirements>())
 			{
 				FMassArchetypeHandle Archetype = EntityManager.GetArchetypeForEntity(Entity);
-				return Requirements->DoesArchetypeMatchRequirements(Archetype);
+				if (Archetype.IsValid())
+				{
+					return Requirements->DoesArchetypeMatchRequirements(Archetype);
+				}
 			}
 			return false;
 		}
