[MassDebugger] added SparseElement bitset to the archetype details
Note that the information exposed is the aggregated info from all chunks; it doesn't mean that all entities have the element.

#Synced-CL 51261467

diff --git a/Engine/Source/Editor/MassEntityDebugger/Private/MassDebuggerModel.cpp b/Engine/Source/Editor/MassEntityDebugger/Private/MassDebuggerModel.cpp
--- a/Engine/Source/Editor/MassEntityDebugger/Private/MassDebuggerModel.cpp
+++ b/Engine/Source/Editor/MassEntityDebugger/Private/MassDebuggerModel.cpp
@@ -254,6 +254,7 @@ FMassDebuggerArchetypeData::FMassDebuggerArchetypeData(const FMassArchetypeHandl
 #if WITH_MASSENTITY_DEBUG
 	Handle = ArchetypeHandle;
 	Composition = FMassDebugger::GetArchetypeComposition(ArchetypeHandle);
+	SparseElementsBitset = FMassDebugger::GetArchetypeSparseElementsBitset(ArchetypeHandle);
 
 	// @todo should ensure we're using same hashing as the EntityManager here
 	CompositionHash = Composition.CalculateHash();
diff --git a/Engine/Source/Editor/MassEntityDebugger/Private/MassDebuggerModel.h b/Engine/Source/Editor/MassEntityDebugger/Private/MassDebuggerModel.h
--- a/Engine/Source/Editor/MassEntityDebugger/Private/MassDebuggerModel.h
+++ b/Engine/Source/Editor/MassEntityDebugger/Private/MassDebuggerModel.h
@@ -143,6 +143,8 @@ struct FMassDebuggerArchetypeData
 
 	FMassArchetypeCompositionDescriptor Composition;
 
+	FMassElementBitSet SparseElementsBitset;
+
 	/** Hash of the Compositions. */
 	uint32 CompositionHash = 0;
 	/** Combined hash of composition and shared fragments. */
diff --git a/Engine/Source/Editor/MassEntityDebugger/Private/SMassArchetype.cpp b/Engine/Source/Editor/MassEntityDebugger/Private/SMassArchetype.cpp
--- a/Engine/Source/Editor/MassEntityDebugger/Private/SMassArchetype.cpp
+++ b/Engine/Source/Editor/MassEntityDebugger/Private/SMassArchetype.cpp
@@ -59,7 +59,7 @@ void SMassArchetype::Construct(const FArguments& InArgs, TSharedPtr<FMassDebugge
 				.Font(FCoreStyle::GetDefaultFontStyle("Mono", 9))
 		];
 
-	FText ArchetypeDescription = FText::Format(LOCTEXT("ArchetypeDescrption", "EntitiesCount: {0}\
+	FText ArchetypeDescription = FText::Format(LOCTEXT("ArchetypeDescription", "EntitiesCount: {0}\
 		\nBytesPerEntity: {1}\
 		\nEntitiesCountPerChunk: {2}\
 		\nChunksCount: {3}\
@@ -71,7 +71,9 @@ void SMassArchetype::Construct(const FArguments& InArgs, TSharedPtr<FMassDebugge
 		, FText::AsNumber(ArchetypeDebugData.ArchetypeStats.ChunksCount)
 		, FText::AsMemory(ArchetypeDebugData.ArchetypeStats.AllocatedSize)
 		, FText::AsMemory(ArchetypeDebugData.ArchetypeStats.WastedEntityMemory)
-		, FText::AsNumber(float(ArchetypeDebugData.ArchetypeStats.WastedEntityMemory) * 100.f / ArchetypeDebugData.ArchetypeStats.AllocatedSize)
+		, FText::AsNumber((ArchetypeDebugData.ArchetypeStats.AllocatedSize != 0
+			? static_cast<float>(ArchetypeDebugData.ArchetypeStats.WastedEntityMemory) * 100.f / ArchetypeDebugData.ArchetypeStats.AllocatedSize
+			: 0.f))
 	);
 
 	Box->AddSlot()
@@ -82,10 +84,10 @@ void SMassArchetype::Construct(const FArguments& InArgs, TSharedPtr<FMassDebugge
 			.Text(ArchetypeDescription)
 		];
 
-	if (ArchetypeDebugData.ArchetypeStats.EntitiesCount != 0 && ArchetypeDebugData.ArchetypeStats.ChunksCount != 0)
+	if (ArchetypeDebugData.ArchetypeStats.EntitiesCountPerChunk != 0 && ArchetypeDebugData.ArchetypeStats.ChunksCount != 0)
 	{
 		const float AverageEntitiesPerChunkActual = float(ArchetypeDebugData.ArchetypeStats.EntitiesCount) / ArchetypeDebugData.ArchetypeStats.ChunksCount;
-		FText DerivedArchetypeDescription = FText::Format(LOCTEXT("ArchetypeDescrptionAux", "Actual average Entities per Chunk: {0}\nChunk occupancy: {1}")
+		FText DerivedArchetypeDescription = FText::Format(LOCTEXT("ArchetypeDescriptionAux", "Actual average Entities per Chunk: {0}\nChunk occupancy: {1}")
 			, FText::AsNumber(AverageEntitiesPerChunkActual)
 			, FText::AsNumber(AverageEntitiesPerChunkActual / ArchetypeDebugData.ArchetypeStats.EntitiesCountPerChunk));
 
@@ -120,6 +122,8 @@ void SMassArchetype::Construct(const FArguments& InArgs, TSharedPtr<FMassDebugge
 		UE::Mass::Debugger::UI::AddBitSet(Box, Composition.DebugGetConstSharedFragments(), TEXT("Const Shared Fragments"), Brush, InDebuggerModel);
 	}
 
+	UE::Mass::Debugger::UI::AddBitSet(Box, ArchetypeData->SparseElementsBitset, TEXT("Sparse Elements"), Brush, InDebuggerModel);
+
 	TSharedRef<SVerticalBox> MainBox = SNew(SVerticalBox);
 
 	MainBox->AddSlot()
diff --git a/Engine/Source/Runtime/MassEntity/Private/MassDebugger.cpp b/Engine/Source/Runtime/MassEntity/Private/MassDebugger.cpp
--- a/Engine/Source/Runtime/MassEntity/Private/MassDebugger.cpp
+++ b/Engine/Source/Runtime/MassEntity/Private/MassDebugger.cpp
@@ -680,6 +680,12 @@ FMassArchetypeCompositionDescriptor FMassDebugger::GetArchetypeComposition(const
 	return FMassArchetypeCompositionDescriptor(ArchetypeData.GetCompositionBitSet());
 }
 
+FMassElementBitSet FMassDebugger::GetArchetypeSparseElementsBitset(const FMassArchetypeHandle& ArchetypeHandle)
+{
+	const FMassArchetypeData& ArchetypeData = FMassArchetypeHelper::ArchetypeDataFromHandleChecked(ArchetypeHandle);
+	return ArchetypeData.GetSparseElementsBitSet();
+}
+
 uint64 FMassDebugger::GetArchetypeTraceID(const FMassArchetypeData& ArchetypeData)
 {
 	return reinterpret_cast<uint64>(&ArchetypeData);
diff --git a/Engine/Source/Runtime/MassEntity/Public/MassDebugger.h b/Engine/Source/Runtime/MassEntity/Public/MassDebugger.h
--- a/Engine/Source/Runtime/MassEntity/Public/MassDebugger.h
+++ b/Engine/Source/Runtime/MassEntity/Public/MassDebugger.h
@@ -197,6 +197,7 @@ struct FMassDebugger
 	static MASSENTITY_API void ForEachArchetype(const FMassEntityManager& EntityManager, const UE::Mass::Debug::FArchetypeFunction& Function);
 	static MASSENTITY_API TArray<FMassArchetypeHandle> GetAllArchetypes(const FMassEntityManager& EntityManager);
 	static MASSENTITY_API FMassArchetypeCompositionDescriptor GetArchetypeComposition(const FMassArchetypeHandle& ArchetypeHandle);
+	static MASSENTITY_API FMassElementBitSet GetArchetypeSparseElementsBitset(const FMassArchetypeHandle& ArchetypeHandle);
 
 	static MASSENTITY_API uint64 GetArchetypeTraceID(const FMassArchetypeData& ArchetypeData);
 	static MASSENTITY_API uint64 GetArchetypeTraceID(const FMassArchetypeHandle& ArchetypeHandle);
