Add stale model check to entity list refresh guard condition in Mass Entity Debugger

#Synced-CL 51720594

diff --git a/Engine/Source/Editor/MassEntityDebugger/Private/SMassEntitiesView.cpp b/Engine/Source/Editor/MassEntityDebugger/Private/SMassEntitiesView.cpp
--- a/Engine/Source/Editor/MassEntityDebugger/Private/SMassEntitiesView.cpp
+++ b/Engine/Source/Editor/MassEntityDebugger/Private/SMassEntitiesView.cpp
@@ -172,7 +172,10 @@ void SMassEntitiesView::ClearEntities()
 FReply SMassEntitiesView::RefreshEntityList()
 {
 #if WITH_MASSENTITY_DEBUG
-	if (!DebuggerModel.IsValid() || !DebuggerModel->Environment.IsValid() || !DebuggerModel->Environment->EntityManager.IsValid())
+	if (!DebuggerModel.IsValid()
+		|| !DebuggerModel->Environment.IsValid()
+		|| !DebuggerModel->Environment->EntityManager.IsValid()
+		|| DebuggerModel->IsStale())
 	{
 		TempEntityList.Reset();
 		if (EntitiesList.IsValid())
